#UE4 Fix crash in Friends and Chat, if a response comes back when the manager is in an unexpected state, such as when returning to main menu when a request is pending

diff --git a/Engine/Source/Developer/FriendsAndChat/Private/FriendsAndChatManager.cpp b/Engine/Source/Developer/FriendsAndChat/Private/FriendsAndChatManager.cpp
--- a/Engine/Source/Developer/FriendsAndChat/Private/FriendsAndChatManager.cpp
+++ b/Engine/Source/Developer/FriendsAndChat/Private/FriendsAndChatManager.cpp
@@ -832,6 +832,14 @@ void FFriendsAndChatManager::OnSendInviteComplete( int32 LocalPlayer, bool bWasS
 void FFriendsAndChatManager::OnReadFriendsComplete( int32 LocalPlayer, bool bWasSuccessful, const FString& ListName, const FString& ErrorStr )
 {
 	TArray< TSharedRef<FOnlineFriend> > Friends;
+
+	if (ReadListRequests.Num() == 0)
+	{
+		// We got a request back when we weren't expecting it
+		UE_LOG(LogOnline, Warning, TEXT("Got ReadFriends response for list %s at unexpected time"), *ListName);
+		return;
+	}
+
 	FriendsInterface->GetFriendsList( 0, ListName, Friends );
 	
 	if ( Friends.Num() > 0 )
@@ -883,6 +891,13 @@ void FFriendsAndChatManager::OnQueryUserInfoComplete( int32 LocalPlayer, bool bW
 
 	OnlineSubMcp->GetUserInterface()->ClearOnQueryUserInfoCompleteDelegate(0, OnQueryUserInfoCompleteDelegate);
 
+	if (ReadListRequests.Num() == 0)
+	{
+		// We got a request back when we weren't expecting it
+		UE_LOG(LogOnline, Warning, TEXT("Got QueryUserInfo response for player %i at unexpected time"), LocalPlayer);
+		return;
+	}
+
 	EFriendsDisplayLists::Type DisplayList = EFriendsDisplayLists::DefaultDisplay;
 	if ( ReadListRequests[0] != EFriendsLists::Default )
 	{
