SLATE - Fixing up more cases of using bool in IsChecked.

diff --git a/Engine/Source/Developer/FriendsAndChat/Private/UI/Widgets/SFriendsList.cpp b/Engine/Source/Developer/FriendsAndChat/Private/UI/Widgets/SFriendsList.cpp
--- a/Engine/Source/Developer/FriendsAndChat/Private/UI/Widgets/SFriendsList.cpp
+++ b/Engine/Source/Developer/FriendsAndChat/Private/UI/Widgets/SFriendsList.cpp
@@ -82,7 +82,7 @@ void SFriendsList::Construct( const FArguments& InArgs )
 					SAssignNew( DefaultPlayersButton, SCheckBox )
 					.Style( &FriendStyle.CheckBoxToggleStyle )
 					.Type( ESlateCheckBoxType::ToggleButton )
-					.IsChecked( CurrentList == EFriendsDisplayLists::DefaultDisplay )
+					.IsChecked( CurrentList == EFriendsDisplayLists::DefaultDisplay ? ESlateCheckBoxState::Checked : ESlateCheckBoxState::Unchecked )
 					.OnCheckStateChanged( this, &SFriendsList::DefaultListSelect_OnChanged )
 					.ForegroundColor( this, &SFriendsList::GetMenuTextColor, EFriendsDisplayLists::DefaultDisplay )
 					[
@@ -97,7 +97,7 @@ void SFriendsList::Construct( const FArguments& InArgs )
 					SAssignNew( RecentPlayersButton, SCheckBox )
 					.Style( &FriendStyle.CheckBoxToggleStyle )
 					.Type( ESlateCheckBoxType::ToggleButton )
-					.IsChecked( CurrentList == EFriendsDisplayLists::RecentPlayersDisplay )
+					.IsChecked(CurrentList == EFriendsDisplayLists::RecentPlayersDisplay ? ESlateCheckBoxState::Checked : ESlateCheckBoxState::Unchecked )
 					.OnCheckStateChanged(this, &SFriendsList::RecentListSelect_OnChanged )
 					.ForegroundColor( this, &SFriendsList::GetMenuTextColor, EFriendsDisplayLists::RecentPlayersDisplay )
 					[
@@ -112,7 +112,7 @@ void SFriendsList::Construct( const FArguments& InArgs )
 					SAssignNew( FriendRequestsButton, SCheckBox )
 					.Style( &FriendStyle.CheckBoxToggleStyle )
 					.Type( ESlateCheckBoxType::ToggleButton )
-					.IsChecked( CurrentList == EFriendsDisplayLists::FriendRequestsDisplay )
+					.IsChecked(CurrentList == EFriendsDisplayLists::FriendRequestsDisplay ? ESlateCheckBoxState::Checked : ESlateCheckBoxState::Unchecked )
 					.OnCheckStateChanged(this, &SFriendsList::RequestListSelect_OnChanged )
 					.ForegroundColor( this, &SFriendsList::GetMenuTextColor, EFriendsDisplayLists::FriendRequestsDisplay )
 					[
