#uelauncher
- fixes for compiling Launcher in non-unity mode

diff --git a/Engine/Source/Developer/FriendsAndChat/Public/Models/FriendsAndChatMessage.h b/Engine/Source/Developer/FriendsAndChat/Public/Models/FriendsAndChatMessage.h
--- a/Engine/Source/Developer/FriendsAndChat/Public/Models/FriendsAndChatMessage.h
+++ b/Engine/Source/Developer/FriendsAndChat/Public/Models/FriendsAndChatMessage.h
@@ -6,6 +6,8 @@
 
 #pragma once
 
+class FUniqueNetId;
+
 /**
  * Type definition for shared references to instances of ILauncherProfile.
  */
