temporarily backing out until fix for module UClass dummies is in

diff --git a/Engine/Source/Developer/FriendsAndChat/Classes/FriendsAndChatStyleDummy.h b/Engine/Source/Developer/FriendsAndChat/Classes/FriendsAndChatStyleDummy.h
new file mode 100644
--- /dev/null
+++ b/Engine/Source/Developer/FriendsAndChat/Classes/FriendsAndChatStyleDummy.h
@@ -0,0 +1,15 @@
+// Copyright 1998-2014 Epic Games, Inc. All Rights Reserved.
+
+/*=============================================================================
+	FriendsAndChatStyle.h: Declares the FFriendsAndChatStyle interface.
+=============================================================================*/
+#pragma once
+
+#include "FriendsAndChatStyleDummy.generated.h"
+
+// Dummy class to make module ordering work
+UCLASS()
+class FRIENDSANDCHAT_API UFriendsAndChatStyleDummy : public UObject
+{
+	GENERATED_UCLASS_BODY()
+};
diff --git a/Engine/Source/Developer/FriendsAndChat/Private/FriendsAndChatStyleDummy.cpp b/Engine/Source/Developer/FriendsAndChat/Private/FriendsAndChatStyleDummy.cpp
new file mode 100644
--- /dev/null
+++ b/Engine/Source/Developer/FriendsAndChat/Private/FriendsAndChatStyleDummy.cpp
@@ -0,0 +1,13 @@
+// Copyright 1998-2014 Epic Games, Inc. All Rights Reserved.
+
+/*=============================================================================
+	FriendsAndChatStyleDummy.cpp: Implements the FriendsAndChatStyleDummy class.
+=============================================================================*/
+
+#include "FriendsAndChatPrivatePCH.h"
+#include "FriendsAndChatStyleDummy.h"
+
+UFriendsAndChatStyleDummy::UFriendsAndChatStyleDummy(const FPostConstructInitializeProperties& PCIP)
+	: Super(PCIP)
+{
+}
