fix compile errors

diff --git a/Engine/Source/Developer/FriendsAndChat/Private/FriendsMessageManager.cpp b/Engine/Source/Developer/FriendsAndChat/Private/FriendsMessageManager.cpp
--- a/Engine/Source/Developer/FriendsAndChat/Private/FriendsMessageManager.cpp
+++ b/Engine/Source/Developer/FriendsAndChat/Private/FriendsMessageManager.cpp
@@ -52,13 +52,26 @@ public:
 			IOnlineChatPtr ChatInterface = OnlineSub->GetChatInterface();
 			if (ChatInterface.IsValid())
 			{
-				TSharedPtr<FChatRoomInfo> RoomInfo = ChatInterface->GetRoomInfo(*LoggedInUser, FChatRoomId(RoomName));
-				if (RoomInfo.IsValid() &&
-					RoomInfo->IsJoined())
+				if (!RoomName.IsEmpty())
 				{
-					ChatInterface->SendRoomChat(*LoggedInUser, FChatRoomId(RoomName), MsgBody);
+					TSharedPtr<FChatRoomInfo> RoomInfo = ChatInterface->GetRoomInfo(*LoggedInUser, FChatRoomId(RoomName));
+					if (RoomInfo.IsValid() &&
+						RoomInfo->IsJoined())
+					{
+						ChatInterface->SendRoomChat(*LoggedInUser, FChatRoomId(RoomName), MsgBody);
+					}
 				}
-				
+				else
+				{
+					// send to all joined rooms
+					TArray<FChatRoomId> JoinedRooms;
+					OnlineSub->GetChatInterface()->GetJoinedRooms(*LoggedInUser, JoinedRooms);
+					for (auto RoomId : JoinedRooms)
+					{
+						ChatInterface->SendRoomChat(*LoggedInUser, RoomId, MsgBody);
+					}
+				}
+
 			}
 		}
 	}
@@ -164,16 +177,7 @@ private:
 	}
 
 	void OnChatRoomJoinPublic(const FUniqueNetId& UserId, const FChatRoomId& ChatRoomID, bool bWasSuccessful, const FString& Error)
-	{		
-		if (bWasSuccessful)
-		{
-			TSharedPtr< FFriendChatMessage > ChatItem = MakeShareable(new FFriendChatMessage());
-			ChatItem->Message = FText::FromString(FString::Printf(TEXT("Enter %s channel"), *ChatRoomID));
-			ChatItem->MessageType = EChatMessageType::Global;
-			ChatItem->MessageTimeText = FText::AsTime(FDateTime::Now());
-			ChatItem->bIsFromSelf = true;
-			OnChatMessageRecieved().Broadcast(ChatItem.ToSharedRef());
-		}
+	{
 	}
 
 	void OnChatRoomExit(const FUniqueNetId& UserId, const FChatRoomId& ChatRoomID, bool bWasSuccessful, const FString& Error)
@@ -226,23 +230,6 @@ private:
 	{
 		TSharedPtr< FFriendChatMessage > ChatItem = MakeShareable(new FFriendChatMessage());
 
-		TSharedPtr<FChatRoomMember> ChatMember = OnlineSub->GetChatInterface()->GetMember(UserId, ChatRoomID, ChatMessage.GetUserId());
-		if (ChatMember.IsValid())
-		{
-			ChatItem->FromName = FText::FromString(*ChatMember->GetNickname());
-		}
-		else
-		{
-			ChatItem->FromName = FText::FromString("Unknown");
-		}
-		ChatItem->Message = FText::FromString(*ChatMessage.GetBody());
-		ChatItem->MessageType = EChatMessageType::Global;
-		ChatItem->MessageTimeText = FText::AsTime(ChatMessage.GetTimestamp());
-		ChatItem->bIsFromSelf = UserId == *LoggedInUser;
-		OnChatMessageRecieved().Broadcast(ChatItem.ToSharedRef());
-
-		TSharedPtr< FFriendChatMessage > ChatItem = MakeShareable(new FFriendChatMessage());
-
 		ChatItem->FromName = FText::FromString(*ChatMessage.GetNickname());
 		ChatItem->Message = FText::FromString(*ChatMessage.GetBody());
 		ChatItem->MessageType = EChatMessageType::Global;
@@ -300,4 +287,4 @@ TSharedRef< FFriendsMessageManager > FFriendsMessageManagerFactory::Create()
 	return MessageManager;
 }
 
-#undef LOCTEXT_NAMESPACE
+#undef LOCTEXT_NAMESPACE
\ No newline at end of file
