Integrate missing files from CL 2361265

diff --git a/Engine/Source/Developer/FriendsAndChat/Private/Models/ChatViewModel.cpp b/Engine/Source/Developer/FriendsAndChat/Private/Models/ChatViewModel.cpp
--- a/Engine/Source/Developer/FriendsAndChat/Private/Models/ChatViewModel.cpp
+++ b/Engine/Source/Developer/FriendsAndChat/Private/Models/ChatViewModel.cpp
@@ -202,6 +202,16 @@ public:
 		return bAllowGlobalChat;
 	}
 
+	virtual void SetCaptureFocus(bool bInCaptureFocus) override
+	{
+		bCaptureFocus = bInCaptureFocus;
+	}
+
+	virtual const bool ShouldCaptureFocus() const override
+	{
+		return bCaptureFocus;
+	}
+
 	virtual void SetAllowGlobalChat(bool bAllow) override
 	{
 		bAllowGlobalChat = bAllow;
@@ -294,6 +304,7 @@ private:
 		, bUseOverrideColor(false)
 		, bInGame(false)
 		, bAllowGlobalChat(true)
+		, bCaptureFocus(false)
 	{
 	}
 
@@ -317,6 +328,7 @@ private:
 	bool bUseOverrideColor;
 	bool bInGame;
 	bool bAllowGlobalChat;
+	bool bCaptureFocus;
 
 private:
 	friend FChatViewModelFactory;
diff --git a/Engine/Source/Developer/FriendsAndChat/Private/Models/ChatViewModel.h b/Engine/Source/Developer/FriendsAndChat/Private/Models/ChatViewModel.h
--- a/Engine/Source/Developer/FriendsAndChat/Private/Models/ChatViewModel.h
+++ b/Engine/Source/Developer/FriendsAndChat/Private/Models/ChatViewModel.h
@@ -36,6 +36,8 @@ public:
 	virtual EVisibility GetSScrollBarVisibility() const = 0;
 	virtual void SetAllowGlobalChat(bool bAllow) = 0;
 	virtual bool IsGlobalChatEnabled() const = 0;
+	virtual void SetCaptureFocus(bool bCaptureFocus) = 0;
+	virtual const bool ShouldCaptureFocus() const = 0;
 
 	DECLARE_EVENT(FChatViewModel, FChatListSetFocus)
 	virtual FChatListSetFocus& OnChatListSetFocus() = 0;
diff --git a/Engine/Source/Developer/FriendsAndChat/Private/UI/Widgets/SChatWindow.cpp b/Engine/Source/Developer/FriendsAndChat/Private/UI/Widgets/SChatWindow.cpp
--- a/Engine/Source/Developer/FriendsAndChat/Private/UI/Widgets/SChatWindow.cpp
+++ b/Engine/Source/Developer/FriendsAndChat/Private/UI/Widgets/SChatWindow.cpp
@@ -444,7 +444,12 @@ private:
 
 	virtual FReply OnFocusReceived( const FGeometry& MyGeometry, const FFocusEvent& InFocusEvent ) override
 	{
-		return FReply::Handled().ReleaseMouseCapture().LockMouseToWidget( SharedThis( this ) );
+		if(ViewModel->ShouldCaptureFocus())
+		{
+			return FReply::Handled().ReleaseMouseCapture();
+		}
+		
+		return FReply::Handled();
 	}
 
 private:
