Merging using UE4-To-UE4-LauncherReleases

diff --git a/Engine/Source/Developer/FriendsAndChat/Private/FriendsAndChatManager.cpp b/Engine/Source/Developer/FriendsAndChat/Private/FriendsAndChatManager.cpp
--- a/Engine/Source/Developer/FriendsAndChat/Private/FriendsAndChatManager.cpp
+++ b/Engine/Source/Developer/FriendsAndChat/Private/FriendsAndChatManager.cpp
@@ -219,31 +219,32 @@ void FFriendsAndChatManager::CreateFriendsListWidget(const FFriendsAndChatStyle*
 
 	Style = *InStyle;
 
-	if (!FriendWindow.IsValid())
+	if ( !FriendWindow.IsValid() )
 	{
 		FriendWindow = SNew( SWindow )
-		.Title(LOCTEXT("FFriendsAndChatManager_FriendsTitle", "Friends List"))
-		.ClientSize(DEFAULT_WINDOW_SIZE)
-		.ScreenPosition(FVector2D(100, 100))
+		.Title( LOCTEXT( "FFriendsAndChatManager_FriendsTitle", "Friends List") )
+		.ClientSize( DEFAULT_WINDOW_SIZE )
+		.ScreenPosition( FVector2D( 100, 100 ) )
 		.AutoCenter( EAutoCenter::None )
 		.SizingRule( ESizingRule::FixedSize )
-		.SupportsMaximize(false)
+		.SupportsMaximize( false )
 		.SupportsMinimize(true)
-		.bDragAnywhere(true)
-		.CreateTitleBar(false);
+		.bDragAnywhere( true )
+		.CreateTitleBar( false );
 
 		FriendWindow->SetOnWindowClosed(FOnWindowClosed::CreateRaw(this, &FFriendsAndChatManager::HandleFriendsWindowClosed));
 
 		BuildFriendsUI();
 		FriendWindow = FSlateApplication::Get().AddWindow(FriendWindow.ToSharedRef());
 	}
-	else if(!FriendWindow->IsWindowMinimized())
+	else if(FriendWindow->IsWindowMinimized())
 	{
 		FriendWindow->Restore();
 		BuildFriendsUI();
 	}
 
 	GenerateChatWindow();
+	FriendWindow->BringToFront(true);
 
 	// Clear notifications
 	OnFriendsNotification().Broadcast(false);
@@ -293,6 +294,8 @@ void FFriendsAndChatManager::BuildFriendsUI()
 			]
 		]
 	);
+
+	FriendWindow->SetTitleBar(TitleBar);
 }
 
 
@@ -362,36 +365,36 @@ void FFriendsAndChatManager::GenerateChatWindow()
 
 	ChatViewModel->SetCaptureFocus(false);
 
-	if (!ChatWindow.IsValid())
-	{
-		ChatWindow = SNew( SWindow )
-		.Title( LOCTEXT( "FriendsAndChatManager_ChatTitle", "Chat Window") )
-		.ClientSize( DEFAULT_WINDOW_SIZE )
-		.ScreenPosition( FVector2D( 200, 100 ) )
-		.AutoCenter( EAutoCenter::None )
+		if (!ChatWindow.IsValid())
+		{
+			ChatWindow = SNew( SWindow )
+			.Title( LOCTEXT( "FriendsAndChatManager_ChatTitle", "Chat Window") )
+			.ClientSize( DEFAULT_WINDOW_SIZE )
+			.ScreenPosition( FVector2D( 200, 100 ) )
+			.AutoCenter( EAutoCenter::None )
 		.SupportsMaximize( false )
-		.SupportsMinimize( true )
-		.CreateTitleBar( false )
-		.SizingRule( ESizingRule::FixedSize );
+			.SupportsMinimize( true )
+			.CreateTitleBar( false )
+			.SizingRule( ESizingRule::FixedSize );
 
-		ChatWindow->SetOnWindowClosed(FOnWindowClosed::CreateRaw(this, &FFriendsAndChatManager::HandleChatWindowClosed));
+			ChatWindow->SetOnWindowClosed(FOnWindowClosed::CreateRaw(this, &FFriendsAndChatManager::HandleChatWindowClosed));
 
-		SetChatWindowContents();
+			SetChatWindowContents();
 		ChatWindow = FSlateApplication::Get().AddWindow(ChatWindow.ToSharedRef());
+		}
+		else if(ChatWindow->IsWindowMinimized())
+		{
+			ChatWindow->Restore();
+			SetChatWindowContents();
+		}
 	}
-	else if(ChatWindow->IsWindowMinimized())
-	{
-		ChatWindow->Restore();
-		SetChatWindowContents();
-	}
-}
 
 void FFriendsAndChatManager::SetChatFriend( TSharedPtr< IFriendItem > FriendItem )
 {
 	if(ChatViewModel.IsValid())
 	{
-		ChatViewModel->SetChatFriend(FriendItem);
-	}
+	ChatViewModel->SetChatFriend(FriendItem);
+}
 
 	if(bCreateChatWindow)
 	{
@@ -441,20 +444,20 @@ void FFriendsAndChatManager::SetUserIsOnline(EOnlinePresenceState::Type OnlineSt
 		{
 			TSharedPtr<FOnlineUserPresence> CurrentPresence;
 			OnlineSubMcp->GetPresenceInterface()->GetCachedPresence(*UserId, CurrentPresence);
-			FOnlineUserPresenceStatus NewStatus;
+		FOnlineUserPresenceStatus NewStatus;
 			if (CurrentPresence.IsValid())
 			{
 				NewStatus = CurrentPresence->Status;
 			}
 			NewStatus.State = OnlineState;
 			OnlineSubMcp->GetPresenceInterface()->SetPresence(*UserId, NewStatus, OnPresenceUpdatedCompleteDelegate);
-		}
 	}
 }
+}
 
 void FFriendsAndChatManager::AcceptFriend( TSharedPtr< IFriendItem > FriendItem )
 {
-	PendingOutgoingAcceptFriendRequests.Add(FUniqueNetIdString(FriendItem->GetOnlineUser()->GetUserId()->ToString()));
+	PendingOutgoingAcceptFriendRequests.Add( FUniqueNetIdString(FriendItem->GetOnlineUser()->GetUserId()->ToString()) );
 	FriendItem->SetPendingAccept();
 	RefreshList();
 	OnFriendsNotification().Broadcast(false);
@@ -486,8 +489,8 @@ void FFriendsAndChatManager::RequestFriend( const FText& FriendName )
 {
 	if ( !FriendName.IsEmpty() )
 	{
-		FriendByNameRequests.AddUnique(*FriendName.ToString());
-	}
+			FriendByNameRequests.AddUnique( *FriendName.ToString() );
+		}
 }
 
 // Process action responses
@@ -857,18 +860,18 @@ void FFriendsAndChatManager::OnQueryUserInfoComplete( int32 LocalPlayer, bool bW
 		TSharedPtr<FUniqueNetId> UserId = OnlineIdentity->GetUniquePlayerId(0);
 		if (UserId.IsValid())
 		{
-			TArray< TSharedRef<FOnlineRecentPlayer> > Players;
+		TArray< TSharedRef<FOnlineRecentPlayer> > Players;
 			if (FriendsInterface->GetRecentPlayers(*UserId, Players))
+		{
+			for (const auto& RecentPlayer : Players)
 			{
-				for (const auto& RecentPlayer : Players)
-				{
-					TSharedRef<FFriendRecentPlayerItem> RecentPlayerItem = MakeShareable(new FFriendRecentPlayerItem(RecentPlayer));
-					TSharedPtr<FOnlineUser> OnlineUser = OnlineSubMcp->GetUserInterface()->GetUserInfo(LocalPlayer, *RecentPlayer->GetUserId());
-					RecentPlayerItem->SetOnlineUser(OnlineUser);
-					RecentPlayersList.Add(RecentPlayerItem);
-				}
+				TSharedRef<FFriendRecentPlayerItem> RecentPlayerItem = MakeShareable(new FFriendRecentPlayerItem(RecentPlayer));
+				TSharedPtr<FOnlineUser> OnlineUser = OnlineSubMcp->GetUserInterface()->GetUserInfo( LocalPlayer, *RecentPlayer->GetUserId());
+				RecentPlayerItem->SetOnlineUser(OnlineUser);
+				RecentPlayersList.Add(RecentPlayerItem);
 			}
 		}
+		}
 		OnFriendsListUpdated().Broadcast();
 		SetState(EFriendsAndManagerState::Idle);
 	}
@@ -956,12 +959,12 @@ void FFriendsAndChatManager::SendFriendRequests()
 	TSharedPtr<FUniqueNetId> UserId = OnlineIdentity->GetUniquePlayerId(0);
 	if (UserId.IsValid())
 	{
-		for (int32 Index = 0; Index < FriendByNameRequests.Num(); Index++)
-		{
-			OnlinePersonaMcp->QueryUserIdMapping(*UserId, FriendByNameRequests[Index]);
-		}
+	for ( int32 Index = 0; Index < FriendByNameRequests.Num(); Index++ )
+	{
+		OnlinePersonaMcp->QueryUserIdMapping(*UserId, FriendByNameRequests[Index]);
 	}
 }
+}
 
 TSharedPtr< FUniqueNetId > FFriendsAndChatManager::FindUserID( const FString& InUsername )
 {
@@ -1047,8 +1050,8 @@ void FFriendsAndChatManager::OnQueryUserIdMappingComplete(bool bWasSuccessful, c
 		}
 		else
 		{
-			TSharedPtr<FUniqueNetId> FriendId = OnlineIdentity->CreateUniquePlayerId(IdentifiedUserId.ToString());
-			PendingOutgoingFriendRequests.Add(FriendId.ToSharedRef());
+		TSharedPtr<FUniqueNetId> FriendId = OnlineIdentity->CreateUniquePlayerId( IdentifiedUserId.ToString() );
+			PendingOutgoingFriendRequests.Add( FriendId.ToSharedRef() );
 			FriendByNameInvites.AddUnique(DisplayName);
 		}
 	}
