Fix linux compile. Not sure if this is the right logic, but it was this way before a merge clobbered it

diff --git a/Engine/Source/Developer/FriendsAndChat/Private/FriendsAndChatManager.cpp b/Engine/Source/Developer/FriendsAndChat/Private/FriendsAndChatManager.cpp
--- a/Engine/Source/Developer/FriendsAndChat/Private/FriendsAndChatManager.cpp
+++ b/Engine/Source/Developer/FriendsAndChat/Private/FriendsAndChatManager.cpp
@@ -813,7 +813,7 @@ void FFriendsAndChatManager::PreProcessList(const FString& ListName)
 				TSharedPtr< IFriendItem > ExistingFriend = FindUser(Friend->GetUserId().Get());
 				if ( ExistingFriend.IsValid() )
 				{
-					if (Friend->GetInviteStatus() != ExistingFriend->GetOnlineFriend()->GetInviteStatus() || ExistingFriend->IsPendingAccepted() && Friend->GetInviteStatus() == EInviteStatus::Accepted)
+					if (Friend->GetInviteStatus() != ExistingFriend->GetOnlineFriend()->GetInviteStatus() || (ExistingFriend->IsPendingAccepted() && Friend->GetInviteStatus() == EInviteStatus::Accepted))
 					{
 						ExistingFriend->SetOnlineFriend(Friend);
 					}
