Slate : Updated all use cases to use EPopupMethod. Removed method specification altogether wherever possible.

diff --git a/Engine/Source/Developer/FriendsAndChat/Private/FriendsAndChatManager.cpp b/Engine/Source/Developer/FriendsAndChat/Private/FriendsAndChatManager.cpp
--- a/Engine/Source/Developer/FriendsAndChat/Private/FriendsAndChatManager.cpp
+++ b/Engine/Source/Developer/FriendsAndChat/Private/FriendsAndChatManager.cpp
@@ -297,7 +297,7 @@ void FFriendsAndChatManager::BuildFriendsUI()
 				[
 					SNew(SFriendsContainer, FFriendsViewModelFactory::Create(SharedThis(this)))
 					.FriendStyle( &Style )
-					.Method(SMenuAnchor::CreateNewWindow)
+					.Method(EPopupMethod::CreateNewWindow)
 				]
 			]
 			+SOverlay::Slot()
@@ -321,7 +321,7 @@ TSharedPtr< SWidget > FFriendsAndChatManager::GenerateFriendsListWidget( const F
 		[
 			 SNew(SFriendsContainer, FFriendsViewModelFactory::Create(SharedThis(this)))
 			.FriendStyle( &Style )
-			.Method(SMenuAnchor::UseCurrentWindow)
+			.Method(EPopupMethod::UseCurrentWindow)
 		]
 		+SOverlay::Slot()
 		.HAlign(HAlign_Fill)
@@ -347,7 +347,7 @@ TSharedPtr< SWidget > FFriendsAndChatManager::GenerateChatWidget(const FFriendsA
 	Style = *InStyle;
 	SAssignNew(ChatWidget, SChatWindow, ChatViewModel.ToSharedRef())
 	.FriendStyle( &Style )
-	.Method(SMenuAnchor::UseCurrentWindow)
+	.Method(EPopupMethod::UseCurrentWindow)
 	.bInGameUI(bInGameUI);
 	return ChatWidget;
 }
diff --git a/Engine/Source/Developer/FriendsAndChat/Private/UI/Widgets/SChatItem.cpp b/Engine/Source/Developer/FriendsAndChat/Private/UI/Widgets/SChatItem.cpp
--- a/Engine/Source/Developer/FriendsAndChat/Private/UI/Widgets/SChatItem.cpp
+++ b/Engine/Source/Developer/FriendsAndChat/Private/UI/Widgets/SChatItem.cpp
@@ -171,7 +171,7 @@ private:
 	FFriendsAndChatStyle FriendStyle;
 
 	// Holds the menu method - Full screen requires use owning window or crashes.
-	SMenuAnchor::EMethod MenuMethod;
+	EPopupMethod MenuMethod;
 };
 
 TSharedRef<SChatItem> SChatItem::New()
diff --git a/Engine/Source/Developer/FriendsAndChat/Private/UI/Widgets/SChatItem.h b/Engine/Source/Developer/FriendsAndChat/Private/UI/Widgets/SChatItem.h
--- a/Engine/Source/Developer/FriendsAndChat/Private/UI/Widgets/SChatItem.h
+++ b/Engine/Source/Developer/FriendsAndChat/Private/UI/Widgets/SChatItem.h
@@ -8,7 +8,7 @@ public:
 	SLATE_USER_ARGS(SChatItem)
 	{ }
 	SLATE_ARGUMENT( const FFriendsAndChatStyle*, FriendStyle )
-	SLATE_ARGUMENT(SMenuAnchor::EMethod, Method)
+	SLATE_ARGUMENT(EPopupMethod, Method)
 	SLATE_END_ARGS()
 
 	/**
diff --git a/Engine/Source/Developer/FriendsAndChat/Private/UI/Widgets/SChatWindow.cpp b/Engine/Source/Developer/FriendsAndChat/Private/UI/Widgets/SChatWindow.cpp
--- a/Engine/Source/Developer/FriendsAndChat/Private/UI/Widgets/SChatWindow.cpp
+++ b/Engine/Source/Developer/FriendsAndChat/Private/UI/Widgets/SChatWindow.cpp
@@ -78,7 +78,7 @@ public:
 					[
 						SAssignNew(ActionMenu, SMenuAnchor)
 						.Placement(EMenuPlacement::MenuPlacement_AboveAnchor)
-						.Method(SMenuAnchor::UseCurrentWindow)
+						.Method(EPopupMethod::UseCurrentWindow)
 						.OnGetMenuContent(this, &SChatWindowImpl::GetMenuContent)
 						[
 							SNew(SButton)
@@ -140,7 +140,7 @@ public:
 									[
 										SAssignNew(ChatItemActionMenu, SMenuAnchor)
 										.Placement(EMenuPlacement::MenuPlacement_AboveAnchor)
-										.Method(SMenuAnchor::UseCurrentWindow)
+										.Method(EPopupMethod::UseCurrentWindow)
 										.OnGetMenuContent(this, &SChatWindowImpl::GetFriendActionMenu)
 										[
 											SNew(SImage)
@@ -651,7 +651,7 @@ private:
 	FFriendsAndChatStyle FriendStyle;
 
 	// Holds the menu method - Full screen requires use owning window or crashes.
-	SMenuAnchor::EMethod MenuMethod;
+	EPopupMethod MenuMethod;
 
 	// Holds the time transparency.
 	float TimeTransparency;
diff --git a/Engine/Source/Developer/FriendsAndChat/Private/UI/Widgets/SChatWindow.h b/Engine/Source/Developer/FriendsAndChat/Private/UI/Widgets/SChatWindow.h
--- a/Engine/Source/Developer/FriendsAndChat/Private/UI/Widgets/SChatWindow.h
+++ b/Engine/Source/Developer/FriendsAndChat/Private/UI/Widgets/SChatWindow.h
@@ -9,7 +9,7 @@ public:
 	: _MaxChatLength(200)
 	{ }
 		SLATE_ARGUMENT( const FFriendsAndChatStyle*, FriendStyle )
-		SLATE_ARGUMENT(SMenuAnchor::EMethod, Method)
+		SLATE_ARGUMENT(EPopupMethod, Method)
 		SLATE_ARGUMENT(int32, MaxChatLength)
 		SLATE_ARGUMENT(bool, bInGameUI)
 	SLATE_END_ARGS()
diff --git a/Engine/Source/Developer/FriendsAndChat/Private/UI/Widgets/SFriendItem.cpp b/Engine/Source/Developer/FriendsAndChat/Private/UI/Widgets/SFriendItem.cpp
--- a/Engine/Source/Developer/FriendsAndChat/Private/UI/Widgets/SFriendItem.cpp
+++ b/Engine/Source/Developer/FriendsAndChat/Private/UI/Widgets/SFriendItem.cpp
@@ -25,21 +25,21 @@ public:
 				SAssignNew(RemoveConfirmationAnchor, SMenuAnchor)
 				.Method(InArgs._Method)
 				.OnGetMenuContent(this, &SFriendItemImpl::GetRemoveConfirmationContent)
-				.Placement(MenuMethod == SMenuAnchor::UseCurrentWindow ? MenuPlacement_MenuLeft : MenuPlacement_MenuRight)
+				.Placement(MenuMethod == EPopupMethod::UseCurrentWindow ? MenuPlacement_MenuLeft : MenuPlacement_MenuRight)
 			]
 			+ SOverlay::Slot()
 			[
 				SAssignNew(JoinGameConfirmationAnchor, SMenuAnchor)
 				.Method(InArgs._Method)
 				.OnGetMenuContent(this, &SFriendItemImpl::GetJoinGameConfirmationContent)
-				.Placement(MenuMethod == SMenuAnchor::UseCurrentWindow ? MenuPlacement_MenuLeft : MenuPlacement_MenuRight)
+				.Placement(MenuMethod == EPopupMethod::UseCurrentWindow ? MenuPlacement_MenuLeft : MenuPlacement_MenuRight)
 			]
 			+SOverlay::Slot()
 			[
 				SAssignNew(Anchor, SMenuAnchor)
 				.Method(InArgs._Method)
 				.OnGetMenuContent(this, &SFriendItemImpl::GetMenuContent)
-				.Placement(MenuMethod == SMenuAnchor::UseCurrentWindow ? MenuPlacement_MenuLeft : MenuPlacement_MenuRight)
+				.Placement(MenuMethod == EPopupMethod::UseCurrentWindow ? MenuPlacement_MenuLeft : MenuPlacement_MenuRight)
 				.Content()
 				[
 					SNew(SButton)
@@ -338,7 +338,7 @@ private:
 			else
 			{
 				OpenTime -= InDeltaTime;
-				if (OpenTime < 0 || MenuMethod != SMenuAnchor::CreateNewWindow)
+				if (OpenTime < 0 || MenuMethod != EPopupMethod::CreateNewWindow)
 				{
 					Anchor->SetIsOpen(false);
 				}
@@ -359,7 +359,7 @@ private:
 
 	TSharedPtr<SWidget> MenuContent;
 
-	SMenuAnchor::EMethod MenuMethod;
+	EPopupMethod MenuMethod;
 
 	float OpenTime;
 };
diff --git a/Engine/Source/Developer/FriendsAndChat/Private/UI/Widgets/SFriendItem.h b/Engine/Source/Developer/FriendsAndChat/Private/UI/Widgets/SFriendItem.h
--- a/Engine/Source/Developer/FriendsAndChat/Private/UI/Widgets/SFriendItem.h
+++ b/Engine/Source/Developer/FriendsAndChat/Private/UI/Widgets/SFriendItem.h
@@ -8,7 +8,7 @@ public:
 	SLATE_USER_ARGS(SFriendItem)
 	{ }
 	SLATE_ARGUMENT( const FFriendsAndChatStyle*, FriendStyle )
-	SLATE_ARGUMENT(SMenuAnchor::EMethod, Method)
+	SLATE_ARGUMENT(EPopupMethod, Method)
 	SLATE_END_ARGS()
 
 	/**
diff --git a/Engine/Source/Developer/FriendsAndChat/Private/UI/Widgets/SFriendsContainer.cpp b/Engine/Source/Developer/FriendsAndChat/Private/UI/Widgets/SFriendsContainer.cpp
--- a/Engine/Source/Developer/FriendsAndChat/Private/UI/Widgets/SFriendsContainer.cpp
+++ b/Engine/Source/Developer/FriendsAndChat/Private/UI/Widgets/SFriendsContainer.cpp
@@ -194,7 +194,7 @@ private:
 	TSharedPtr<FFriendsViewModel> ViewModel;
 
 	// Holds the menu method - Full screen requires use owning window or crashes.
-	SMenuAnchor::EMethod MenuMethod;
+	EPopupMethod MenuMethod;
 
 	/** Holds the list of friends. */
 	TArray< TSharedPtr< IFriendItem > > FriendsList;
diff --git a/Engine/Source/Developer/FriendsAndChat/Private/UI/Widgets/SFriendsContainer.h b/Engine/Source/Developer/FriendsAndChat/Private/UI/Widgets/SFriendsContainer.h
--- a/Engine/Source/Developer/FriendsAndChat/Private/UI/Widgets/SFriendsContainer.h
+++ b/Engine/Source/Developer/FriendsAndChat/Private/UI/Widgets/SFriendsContainer.h
@@ -9,10 +9,10 @@ class SFriendsContainer : public SUserWidget
 public:
 
 	SLATE_USER_ARGS(SFriendsContainer)
-		: _Method(SMenuAnchor::UseCurrentWindow)
+		: _Method(EPopupMethod::UseCurrentWindow)
 	 { }
 	SLATE_ARGUMENT(const FFriendsAndChatStyle*, FriendStyle)
-	SLATE_ARGUMENT(SMenuAnchor::EMethod, Method)
+	SLATE_ARGUMENT(EPopupMethod, Method)
 	SLATE_END_ARGS()
 
 	virtual void Construct(const FArguments& InArgs, const TSharedRef<FFriendsViewModel>& ViewModel) = 0;
diff --git a/Engine/Source/Developer/FriendsAndChat/Private/UI/Widgets/SFriendsList.cpp b/Engine/Source/Developer/FriendsAndChat/Private/UI/Widgets/SFriendsList.cpp
--- a/Engine/Source/Developer/FriendsAndChat/Private/UI/Widgets/SFriendsList.cpp
+++ b/Engine/Source/Developer/FriendsAndChat/Private/UI/Widgets/SFriendsList.cpp
@@ -53,7 +53,7 @@ private:
 	FFriendsAndChatStyle FriendStyle;
 	TSharedPtr<FFriendListViewModel> ViewModel;
 	TSharedPtr<SVerticalBox> FriendsContents;
-	SMenuAnchor::EMethod MenuMethod;
+	EPopupMethod MenuMethod;
 };
 
 TSharedRef<SFriendsList> SFriendsList::New()
diff --git a/Engine/Source/Developer/FriendsAndChat/Private/UI/Widgets/SFriendsList.h b/Engine/Source/Developer/FriendsAndChat/Private/UI/Widgets/SFriendsList.h
--- a/Engine/Source/Developer/FriendsAndChat/Private/UI/Widgets/SFriendsList.h
+++ b/Engine/Source/Developer/FriendsAndChat/Private/UI/Widgets/SFriendsList.h
@@ -9,7 +9,7 @@ public:
 	SLATE_USER_ARGS(SFriendsList)
 	{ }
 	SLATE_ARGUMENT(const FFriendsAndChatStyle*, FriendStyle)
-	SLATE_ARGUMENT( SMenuAnchor::EMethod, Method )
+	SLATE_ARGUMENT(EPopupMethod, Method)
 	SLATE_END_ARGS()
 
 	/**
diff --git a/Engine/Source/Developer/FriendsAndChat/Private/UI/Widgets/SFriendsListContainer.cpp b/Engine/Source/Developer/FriendsAndChat/Private/UI/Widgets/SFriendsListContainer.cpp
--- a/Engine/Source/Developer/FriendsAndChat/Private/UI/Widgets/SFriendsListContainer.cpp
+++ b/Engine/Source/Developer/FriendsAndChat/Private/UI/Widgets/SFriendsListContainer.cpp
@@ -154,7 +154,7 @@ private:
 
 	TSharedPtr<FFriendListViewModel> ViewModel;
 
-	SMenuAnchor::EMethod MenuMethod;
+	EPopupMethod MenuMethod;
 };
 
 TSharedRef<SFriendsListContainer> SFriendsListContainer::New()
diff --git a/Engine/Source/Developer/FriendsAndChat/Private/UI/Widgets/SFriendsListContainer.h b/Engine/Source/Developer/FriendsAndChat/Private/UI/Widgets/SFriendsListContainer.h
--- a/Engine/Source/Developer/FriendsAndChat/Private/UI/Widgets/SFriendsListContainer.h
+++ b/Engine/Source/Developer/FriendsAndChat/Private/UI/Widgets/SFriendsListContainer.h
@@ -8,7 +8,7 @@ public:
 
 	SLATE_USER_ARGS(SFriendsListContainer) { }
 	SLATE_ARGUMENT(const FFriendsAndChatStyle*, FriendStyle)
-	SLATE_ARGUMENT( SMenuAnchor::EMethod, Method )
+	SLATE_ARGUMENT(EPopupMethod, Method)
 	SLATE_END_ARGS()
 
 	virtual void Construct(const FArguments& InArgs, const TSharedRef<class FFriendListViewModel>& ViewModel) = 0;
diff --git a/Engine/Source/Developer/FriendsAndChat/Private/UI/Widgets/SFriendsStatus.h b/Engine/Source/Developer/FriendsAndChat/Private/UI/Widgets/SFriendsStatus.h
--- a/Engine/Source/Developer/FriendsAndChat/Private/UI/Widgets/SFriendsStatus.h
+++ b/Engine/Source/Developer/FriendsAndChat/Private/UI/Widgets/SFriendsStatus.h
@@ -8,7 +8,7 @@ public:
 
 	SLATE_USER_ARGS(SFriendsStatus) { }
 	SLATE_ARGUMENT(const FFriendsAndChatStyle*, FriendStyle)
-	SLATE_ARGUMENT(SMenuAnchor::EMethod, Method)
+	SLATE_ARGUMENT(EPopupMethod, Method)
 	SLATE_END_ARGS()
 
 	virtual void Construct(const FArguments& InArgs, const TSharedRef<class FFriendsStatusViewModel>& ViewModel) = 0;
diff --git a/Engine/Source/Developer/FriendsAndChat/Private/UI/Widgets/SFriendsUserSettings.h b/Engine/Source/Developer/FriendsAndChat/Private/UI/Widgets/SFriendsUserSettings.h
--- a/Engine/Source/Developer/FriendsAndChat/Private/UI/Widgets/SFriendsUserSettings.h
+++ b/Engine/Source/Developer/FriendsAndChat/Private/UI/Widgets/SFriendsUserSettings.h
@@ -8,7 +8,7 @@ public:
 
 	SLATE_USER_ARGS(SFriendsUserSettings) { }
 	SLATE_ARGUMENT(const FFriendsAndChatStyle*, FriendStyle)
-	SLATE_ARGUMENT(SMenuAnchor::EMethod, Method)
+	SLATE_ARGUMENT(EPopupMethod, Method)
 	SLATE_END_ARGS()
 
 	virtual void Construct(const FArguments& InArgs, const TSharedRef<class FFriendsUserSettingsViewModel>& ViewModel) = 0;
diff --git a/Engine/Source/Developer/SuperSearch/Private/SSuperSearch.cpp b/Engine/Source/Developer/SuperSearch/Private/SSuperSearch.cpp
--- a/Engine/Source/Developer/SuperSearch/Private/SSuperSearch.cpp
+++ b/Engine/Source/Developer/SuperSearch/Private/SSuperSearch.cpp
@@ -47,7 +47,7 @@ void SSuperSearchBox::Construct( const FArguments& InArgs )
 	[
 		SAssignNew( SuggestionBox, SMenuAnchor )
 			.Placement( InArgs._SuggestionListPlacement )
-		.Method( SMenuAnchor::UseCurrentWindow )
+		.Method( EPopupMethod::UseCurrentWindow )
 			[
 				SAssignNew(InputText, SSearchBox)
 					.OnTextCommitted(this, &SSuperSearchBox::OnTextCommitted)
diff --git a/Engine/Source/Editor/LevelEditor/Private/SLevelViewportControlsPopup.cpp b/Engine/Source/Editor/LevelEditor/Private/SLevelViewportControlsPopup.cpp
--- a/Engine/Source/Editor/LevelEditor/Private/SLevelViewportControlsPopup.cpp
+++ b/Engine/Source/Editor/LevelEditor/Private/SLevelViewportControlsPopup.cpp
@@ -39,7 +39,7 @@ void SLevelViewportControlsPopup::Construct(const FArguments& InArgs)
 		ChildSlot
 		[
 			SAssignNew(MenuAnchor, SMenuAnchor)
-			.Method(SMenuAnchor::UseCurrentWindow)
+			.Method(EPopupMethod::UseCurrentWindow)
 			.Placement(MenuPlacement_AboveAnchor)
 			.OnGetMenuContent(this, &SLevelViewportControlsPopup::OnGetMenuContent)
 			[
diff --git a/Engine/Source/Runtime/Slate/Public/Widgets/Input/SMenuAnchor.h b/Engine/Source/Runtime/Slate/Public/Widgets/Input/SMenuAnchor.h
--- a/Engine/Source/Runtime/Slate/Public/Widgets/Input/SMenuAnchor.h
+++ b/Engine/Source/Runtime/Slate/Public/Widgets/Input/SMenuAnchor.h
@@ -48,7 +48,6 @@ public:
 
 		SLATE_ATTRIBUTE( EMenuPlacement, Placement )
 
-		DEPRECATED(4.7, "Popup method now determined automatically; prefer not setting it. Use EPopupMethod when setting explicitly.")
 		SLATE_ARGUMENT(TOptional<EPopupMethod>, Method)
 
 	SLATE_END_ARGS()
diff --git a/Engine/Source/Runtime/UMG/Private/Components/ComboBoxString.cpp b/Engine/Source/Runtime/UMG/Private/Components/ComboBoxString.cpp
--- a/Engine/Source/Runtime/UMG/Private/Components/ComboBoxString.cpp
+++ b/Engine/Source/Runtime/UMG/Private/Components/ComboBoxString.cpp
@@ -47,7 +47,6 @@ TSharedRef<SWidget> UComboBoxString::RebuildWidget()
 		.ContentPadding(ContentPadding)
 		.MaxListHeight(MaxListHeight)
 		.HasDownArrow(HasDownArrow)
-		.Method(SMenuAnchor::UseCurrentWindow)
 		.OnGenerateWidget(BIND_UOBJECT_DELEGATE(SComboBox< TSharedPtr<FString> >::FOnGenerateWidget, HandleGenerateWidget))
 		.OnSelectionChanged(BIND_UOBJECT_DELEGATE(SComboBox< TSharedPtr<FString> >::FOnSelectionChanged, HandleSelectionChanged))
 		.OnComboBoxOpening(BIND_UOBJECT_DELEGATE(FOnComboBoxOpening, HandleOpening))
diff --git a/Engine/Source/Runtime/UMG/Private/Components/MenuAnchor.cpp b/Engine/Source/Runtime/UMG/Private/Components/MenuAnchor.cpp
--- a/Engine/Source/Runtime/UMG/Private/Components/MenuAnchor.cpp
+++ b/Engine/Source/Runtime/UMG/Private/Components/MenuAnchor.cpp
@@ -23,7 +23,6 @@ void UMenuAnchor::ReleaseSlateResources(bool bReleaseChildren)
 TSharedRef<SWidget> UMenuAnchor::RebuildWidget()
 {
 	MyMenuAnchor = SNew(SMenuAnchor)
-		.Method(SMenuAnchor::UseCurrentWindow)
 		.Placement(Placement)
 		.OnGetMenuContent(BIND_UOBJECT_DELEGATE(FOnGetContent, HandleGetMenuContent));
 
