#FriendsAndChat - Update chat view model to be per widget. Should be able to fix bugs more easily with this change

diff --git a/Engine/Source/Developer/FriendsAndChat/Private/FriendsAndChatManager.cpp b/Engine/Source/Developer/FriendsAndChat/Private/FriendsAndChatManager.cpp
--- a/Engine/Source/Developer/FriendsAndChat/Private/FriendsAndChatManager.cpp
+++ b/Engine/Source/Developer/FriendsAndChat/Private/FriendsAndChatManager.cpp
@@ -338,28 +338,22 @@ TSharedPtr< SWidget > FFriendsAndChatManager::GenerateFriendsListWidget( const F
 }
 
 
-TSharedPtr< SWidget > FFriendsAndChatManager::GenerateChatWidget(const FFriendsAndChatStyle* InStyle, bool bInGameUI)
+TSharedPtr< SWidget > FFriendsAndChatManager::GenerateChatWidget(const FFriendsAndChatStyle* InStyle, TSharedRef<IChatViewModel> ViewModel)
 {
-	GetChatViewModel(); // ensure creation
-	TSharedPtr<SChatWindow> ChatWidget;
-	ChatViewModel->SetCaptureFocus(true);
+	// todo - NDavies = find a better way to do this
+	TSharedRef<FChatViewModel> ChatViewModel = StaticCastSharedRef<FChatViewModel>(ViewModel);
 
+	TSharedPtr<SChatWindow> ChatWidget;
 	Style = *InStyle;
-	SAssignNew(ChatWidget, SChatWindow, ChatViewModel.ToSharedRef())
-	.FriendStyle( &Style )
-	.Method(EPopupMethod::UseCurrentWindow)
-	.bInGameUI(bInGameUI);
+	SAssignNew(ChatWidget, SChatWindow, ChatViewModel)
+	.FriendStyle(&Style)
+	.Method(SMenuAnchor::UseCurrentWindow);
 	return ChatWidget;
 }
 
 TSharedPtr<IChatViewModel> FFriendsAndChatManager::GetChatViewModel()
 {
-	if (!ChatViewModel.IsValid())
-	{
-		check(MessageManager.IsValid());
-		ChatViewModel = FChatViewModelFactory::Create(MessageManager.ToSharedRef());
-	}
-	return ChatViewModel;
+	return FChatViewModelFactory::Create(MessageManager.ToSharedRef());
 }
 
 void FFriendsAndChatManager::GenerateChatWindow()
@@ -369,13 +363,6 @@ void FFriendsAndChatManager::GenerateChatWindow()
 	check(MessageManager.IsValid());
 	bCreateChatWindow = true;
 
-	if(!ChatViewModel.IsValid())
-	{
-		ChatViewModel = FChatViewModelFactory::Create(MessageManager.ToSharedRef());
-	}
-
-	ChatViewModel->SetCaptureFocus(false);
-
 	if (!ChatWindow.IsValid())
 	{
 		ChatWindow = SNew( SWindow )
@@ -402,10 +389,7 @@ void FFriendsAndChatManager::GenerateChatWindow()
 
 void FFriendsAndChatManager::SetChatFriend( TSharedPtr< IFriendItem > FriendItem )
 {
-	if(ChatViewModel.IsValid())
-	{
-		ChatViewModel->SetChatFriend(FriendItem);
-	}
+	OnChatFriendSelected().Broadcast(FriendItem);
 
 	if(bCreateChatWindow)
 	{
@@ -423,6 +407,9 @@ void FFriendsAndChatManager::HandleChatWindowClosed(const TSharedRef<SWindow>& I
 void FFriendsAndChatManager::SetChatWindowContents()
 {
 	TSharedPtr<SWindowTitleBar> TitleBar;
+	TSharedPtr<FChatViewModel> ChatViewModel = FChatViewModelFactory::Create(MessageManager.ToSharedRef());
+	ChatViewModel->SetInGameUI(false);
+	ChatViewModel->SetCaptureFocus(false);
 
 	ChatWindow->SetContent(
 		SNew( SBorder )
@@ -443,7 +430,6 @@ void FFriendsAndChatManager::SetChatWindowContents()
 			[
 				SNew(SChatWindow, ChatViewModel.ToSharedRef())
 				.FriendStyle( &Style )
-				.bInGameUI(false)
 			]
 		]);
 }
diff --git a/Engine/Source/Developer/FriendsAndChat/Private/FriendsAndChatManager.h b/Engine/Source/Developer/FriendsAndChat/Private/FriendsAndChatManager.h
--- a/Engine/Source/Developer/FriendsAndChat/Private/FriendsAndChatManager.h
+++ b/Engine/Source/Developer/FriendsAndChat/Private/FriendsAndChatManager.h
@@ -128,7 +128,7 @@ public:
 	virtual void SetUserSettings(const FFriendsAndChatSettings& UserSettings) override;
 	virtual void SetAnalyticsProvider(const TSharedPtr<IAnalyticsProvider>& AnalyticsProvider) override;
 	virtual TSharedPtr< SWidget > GenerateFriendsListWidget( const FFriendsAndChatStyle* InStyle ) override;
-	virtual TSharedPtr< SWidget > GenerateChatWidget(const FFriendsAndChatStyle* InStyle, bool bInGameUI) override;
+	virtual TSharedPtr< SWidget > GenerateChatWidget(const FFriendsAndChatStyle* InStyle, TSharedRef<IChatViewModel> ViewModel) override;
 	virtual TSharedPtr<IChatViewModel> GetChatViewModel() override;
 	virtual void InsertNetworkChatMessage(const FString& InMessage) override;
 	virtual void JoinPublicChatRoom(const FString& RoomName) override;
@@ -348,6 +348,12 @@ public:
 		return OnGameInvitesUpdatedDelegate;
 	}
 
+	DECLARE_EVENT_OneParam(FFriendsAndChatManager, FOnChatFriendSelected, TSharedPtr<IFriendItem> /*ChatFriend*/)
+	virtual FOnChatFriendSelected& OnChatFriendSelected()
+	{
+		return OnChatFriendSelectedDelegate;
+	}
+
 private:
 
 	/** Process the friends list to send out invites */
@@ -657,6 +663,8 @@ private:
 	FOnFriendsUpdated OnFriendsListUpdatedDelegate;
 	// Delegate for when list of active game invites updates
 	FOnGameInvitesUpdated OnGameInvitesUpdatedDelegate;
+	// Delegate for when a friend is selected for chat
+	FOnChatFriendSelected OnChatFriendSelectedDelegate;
 
 	/* Identity stuff
 	*****************************************************************************/
@@ -675,8 +683,6 @@ private:
 	TArray<FString> ChatRoomstoJoin;
 	// Manages private/public chat messages 
 	TSharedPtr<class FFriendsMessageManager> MessageManager;
-	// Info needed to view chat messages
-	TSharedPtr<class FChatViewModel> ChatViewModel;
 
 	/* Manger state
 	*****************************************************************************/
diff --git a/Engine/Source/Developer/FriendsAndChat/Private/FriendsMessageManager.cpp b/Engine/Source/Developer/FriendsAndChat/Private/FriendsMessageManager.cpp
--- a/Engine/Source/Developer/FriendsAndChat/Private/FriendsMessageManager.cpp
+++ b/Engine/Source/Developer/FriendsAndChat/Private/FriendsMessageManager.cpp
@@ -44,6 +44,11 @@ public:
 		UnInitialize();
 	}
 
+	virtual const TArray<TSharedPtr<FFriendChatMessage> >& GetMessageList() const override
+	{
+		return ReceivedMessages;
+	}
+
 	virtual bool SendRoomMessage(const FString& RoomName, const FString& MsgBody) override
 	{
 		if (OnlineSub != nullptr &&
@@ -82,7 +87,7 @@ public:
 		return false;
 	}
 
-	virtual bool SendPrivateMessage(const FUniqueNetId& RecipientId, const FString& MsgBody) override
+	virtual bool SendPrivateMessage(const FUniqueNetId& RecipientId, TSharedPtr< FFriendChatMessage > ChatMessage) override
 	{
 		if (OnlineSub != nullptr &&
 			LoggedInUser.IsValid())
@@ -90,7 +95,8 @@ public:
 			IOnlineChatPtr ChatInterface = OnlineSub->GetChatInterface();
 			if(ChatInterface.IsValid())
 			{
-				return ChatInterface->SendPrivateChat(*LoggedInUser, RecipientId, MsgBody);
+				AddMessage(ChatMessage.ToSharedRef());
+				return ChatInterface->SendPrivateChat(*LoggedInUser, RecipientId, ChatMessage->Message.ToString());
 			}
 		}
 		return false;
@@ -104,7 +110,7 @@ public:
 		ChatItem->MessageType = EChatMessageType::Party;
 		ChatItem->MessageTimeText = FText::AsTime(FDateTime::UtcNow());
 		ChatItem->bIsFromSelf = false;
-		OnChatMessageRecieved().Broadcast(ChatItem.ToSharedRef());
+		AddMessage(ChatItem.ToSharedRef());
 	}
 
 	virtual void JoinPublicRoom(const FString& RoomName) override
@@ -236,7 +242,7 @@ private:
 			ChatItem->MessageType = EChatMessageType::Global;
 			ChatItem->MessageTimeText = FText::AsTime(FDateTime::UtcNow());
 			ChatItem->bIsFromSelf = false;
-			OnChatMessageRecieved().Broadcast(ChatItem.ToSharedRef());
+			AddMessage(ChatItem);
 		}
 	}
 
@@ -256,7 +262,7 @@ private:
 			ChatItem->MessageType = EChatMessageType::Global;
 			ChatItem->MessageTimeText = FText::AsTime(FDateTime::UtcNow());
 			ChatItem->bIsFromSelf = false;
-			OnChatMessageRecieved().Broadcast(ChatItem.ToSharedRef());
+			AddMessage(ChatItem);
 		}
 	}
 
@@ -280,7 +286,7 @@ private:
 		}
 
 		ChatItem->MessageRef = ChatMessage;
-		OnChatMessageRecieved().Broadcast(ChatItem.ToSharedRef());
+		AddMessage(ChatItem);
 
 	}
 
@@ -299,7 +305,7 @@ private:
 			ChatItem->MessageTimeText = FText::AsTime(ChatMessage->GetTimestamp());
 			ChatItem->bIsFromSelf = false;
 			ChatItem->MessageRef = ChatMessage;
-			OnChatMessageRecieved().Broadcast(ChatItem.ToSharedRef());
+			AddMessage(ChatItem);
 		}
 	}
 
@@ -318,6 +324,12 @@ private:
 		}
 	}
 
+	void AddMessage(TSharedPtr< FFriendChatMessage > ChatItem)
+	{
+		ReceivedMessages.Add(ChatItem);
+		OnChatMessageRecieved().Broadcast(ChatItem.ToSharedRef());
+	}
+
 private:
 
 	// Incoming delegates
@@ -339,6 +351,8 @@ private:
 	TSharedPtr<FUniqueNetId> LoggedInUser;
 	TArray<FString> RoomJoins;
 
+	TArray<TSharedPtr< FFriendChatMessage > > ReceivedMessages; 
+
 	bool bEnableEnterExitMessages;
 
 private:
diff --git a/Engine/Source/Developer/FriendsAndChat/Private/FriendsMessageManager.h b/Engine/Source/Developer/FriendsAndChat/Private/FriendsMessageManager.h
--- a/Engine/Source/Developer/FriendsAndChat/Private/FriendsMessageManager.h
+++ b/Engine/Source/Developer/FriendsAndChat/Private/FriendsMessageManager.h
@@ -59,9 +59,10 @@ public:
 
 	virtual void LogIn() = 0;
 	virtual void LogOut() = 0;
+	virtual const TArray<TSharedPtr<FFriendChatMessage> >& GetMessageList() const = 0;
 	virtual void JoinPublicRoom(const FString& RoomName) = 0;
 	virtual bool SendRoomMessage(const FString& RoomName, const FString& MsgBody) = 0;
-	virtual bool SendPrivateMessage(const FUniqueNetId& RecipientId, const FString& MsgBody) = 0;
+	virtual bool SendPrivateMessage(const FUniqueNetId& RecipientId, TSharedPtr< FFriendChatMessage > ChatMessage) = 0;
 	virtual void InsertNetworkMessage(const FString& MsgBody) = 0;
 
 	DECLARE_EVENT_OneParam(FFriendsMessageManager, FOnChatMessageReceivedEvent, const TSharedRef<FFriendChatMessage> /*The chat message*/)
diff --git a/Engine/Source/Developer/FriendsAndChat/Private/Models/ChatViewModel.cpp b/Engine/Source/Developer/FriendsAndChat/Private/Models/ChatViewModel.cpp
--- a/Engine/Source/Developer/FriendsAndChat/Private/Models/ChatViewModel.cpp
+++ b/Engine/Source/Developer/FriendsAndChat/Private/Models/ChatViewModel.cpp
@@ -231,17 +231,16 @@ public:
 				{
 					if (SelectedFriend.IsValid() && SelectedFriend->UserID.IsValid())
 					{
-						if (MessageManager.Pin()->SendPrivateMessage(*SelectedFriend->UserID.Get(), NewMessage.ToString()))
+						TSharedPtr< FFriendChatMessage > ChatItem = MakeShareable(new FFriendChatMessage());
+						ChatItem->FromName = SelectedFriend->FriendName;
+						ChatItem->Message = NewMessage;
+						ChatItem->MessageType = EChatMessageType::Whisper;
+						ChatItem->MessageTimeText = FText::AsTime(FDateTime::UtcNow());
+						ChatItem->bIsFromSelf = true;
+						ChatItem->SenderId = SelectedFriend->UserID;
+
+						if (MessageManager.Pin()->SendPrivateMessage(*SelectedFriend->UserID.Get(), ChatItem))
 						{
-							TSharedPtr< FFriendChatMessage > ChatItem = MakeShareable(new FFriendChatMessage());
-							ChatItem->FromName = SelectedFriend->FriendName;
-							ChatItem->Message = NewMessage;
-							ChatItem->MessageType = EChatMessageType::Whisper;
-							ChatItem->MessageTimeText = FText::AsTime(FDateTime::UtcNow());
-							ChatItem->bIsFromSelf = true;
-							ChatItem->SenderId = SelectedFriend->UserID;
-							ChatLists.Add(FChatItemViewModelFactory::Create(ChatItem.ToSharedRef(), SharedThis(this)));
-							FilterChatList();
 							bSuccess = true;
 						}
 
@@ -313,17 +312,14 @@ public:
 		return EVisibility::Collapsed;
 	}
 
-	virtual EVisibility GetInviteToGameVisibility() const override
+	virtual EVisibility GetScrollbarVisibility() const override
 	{
-		return bAllowJoinGame ? EVisibility::Visible : EVisibility::Collapsed;
+		return bInGame ? EVisibility::Collapsed : EVisibility::Visible;
 	}
 
-	virtual void SetChatFriend(TSharedPtr<IFriendItem> ChatFriend) override
+	virtual EVisibility GetInviteToGameVisibility() const override
 	{
-		if(ChatFriend.IsValid())
-		{
-			SetWhisperChannel(GetRecentFriend(FText::FromString(ChatFriend->GetName()), ChatFriend->GetUniqueID()));
-		}
+		return bAllowJoinGame ? EVisibility::Visible : EVisibility::Collapsed;
 	}
 
 	virtual bool IsGlobalChatEnabled() const override
@@ -394,8 +390,13 @@ private:
 	void Initialize()
 	{
 		FFriendsAndChatManager::Get()->OnFriendsListUpdated().AddSP(this, &FChatViewModelImpl::UpdateSelectedFriendsViewModel);
+		FFriendsAndChatManager::Get()->OnChatFriendSelected().AddSP(this, &FChatViewModelImpl::HandleChatFriendSelected);
 		MessageManager.Pin()->OnChatMessageRecieved().AddSP(this, &FChatViewModelImpl::HandleMessageReceived);
-		FilterChatList();
+
+		for( const auto& Message : MessageManager.Pin()->GetMessageList())
+		{
+			HandleMessageReceived(Message.ToSharedRef());
+		}
 	}
 
 	void FilterChatList()
@@ -437,6 +438,14 @@ private:
 		FilterChatList();
 	}
 
+	void HandleChatFriendSelected(TSharedPtr<IFriendItem> ChatFriend)
+	{
+		if(ChatFriend.IsValid())
+		{
+			SetWhisperChannel(GetRecentFriend(FText::FromString(ChatFriend->GetName()), ChatFriend->GetUniqueID()));
+		}
+	}
+
 	void HandleSetFocus()
 	{
 		OnChatListSetFocus().Broadcast();
diff --git a/Engine/Source/Developer/FriendsAndChat/Private/Models/ChatViewModel.h b/Engine/Source/Developer/FriendsAndChat/Private/Models/ChatViewModel.h
--- a/Engine/Source/Developer/FriendsAndChat/Private/Models/ChatViewModel.h
+++ b/Engine/Source/Developer/FriendsAndChat/Private/Models/ChatViewModel.h
@@ -30,6 +30,7 @@ public:
 	virtual FText GetChatGroupText() const = 0;
 	virtual const EVisibility GetTextEntryVisibility() = 0;
 	virtual const EVisibility GetConfirmationVisibility() = 0;
+	virtual EVisibility GetScrollbarVisibility() const = 0;
 	virtual EVisibility GetInviteToGameVisibility() const = 0;
 	virtual void EnumerateChatChannelOptionsList(TArray<EChatMessageType::Type>& OUTChannelType) = 0;
 	virtual void EnumerateFriendOptions(TArray<EFriendActionType::Type>& OUTActionList) = 0;
@@ -44,7 +45,6 @@ public:
 	virtual const float GetTimeTransparency() const = 0;
 	virtual EChatMessageType::Type GetChatChannelType() const = 0;
 	virtual const TArray<TSharedPtr<FSelectedFriend> >& GetRecentOptions() const = 0;
-	virtual void SetChatFriend(TSharedPtr<IFriendItem> ChatFriend) = 0;
 	virtual void SetAllowGlobalChat(bool bAllow) = 0;
 	virtual bool IsGlobalChatEnabled() const = 0;
 	virtual const bool IsChatHidden() = 0;
diff --git a/Engine/Source/Developer/FriendsAndChat/Private/UI/Widgets/SChatWindow.cpp b/Engine/Source/Developer/FriendsAndChat/Private/UI/Widgets/SChatWindow.cpp
--- a/Engine/Source/Developer/FriendsAndChat/Private/UI/Widgets/SChatWindow.cpp
+++ b/Engine/Source/Developer/FriendsAndChat/Private/UI/Widgets/SChatWindow.cpp
@@ -16,7 +16,6 @@ public:
 	{
 		FriendStyle = *InArgs._FriendStyle;
 		MenuMethod = InArgs._Method;
-		bInGameUI = InArgs._bInGameUI;
 		MaxChatLength = InArgs._MaxChatLength;
 		this->ViewModel = InViewModel;
 		FChatViewModel* ViewModelPtr = ViewModel.Get();
@@ -46,7 +45,7 @@ public:
 			 		.AutoWidth()
 					[
 						SNew( SBorder )
-						.Visibility(bInGameUI ? EVisibility::Collapsed : EVisibility::Visible)
+						.Visibility(ViewModel->GetScrollbarVisibility())
 						.BorderBackgroundColor(this, &SChatWindowImpl::GetTimedFadeSlateColor)
 						.ColorAndOpacity(this, &SChatWindowImpl::GetTimedFadeColor)
 						[
@@ -559,7 +558,7 @@ private:
 
 	EVisibility GetActionVisibility() const
 	{
-		return !bInGameUI ? EVisibility::Visible : ViewModel->GetEntryBarVisibility();
+		return ViewModel->GetEntryBarVisibility();
 	}
 
 	EVisibility GetChatEntryVisibility() const
@@ -652,9 +651,6 @@ private:
 
 	// Holds the time transparency.
 	float TimeTransparency;
-
-	// Holds if this widget is being displayed in game
-	bool bInGameUI;
 };
 
 TSharedRef<SChatWindow> SChatWindow::New()
diff --git a/Engine/Source/Developer/FriendsAndChat/Private/UI/Widgets/SChatWindow.h b/Engine/Source/Developer/FriendsAndChat/Private/UI/Widgets/SChatWindow.h
--- a/Engine/Source/Developer/FriendsAndChat/Private/UI/Widgets/SChatWindow.h
+++ b/Engine/Source/Developer/FriendsAndChat/Private/UI/Widgets/SChatWindow.h
@@ -11,7 +11,6 @@ public:
 		SLATE_ARGUMENT( const FFriendsAndChatStyle*, FriendStyle )
 		SLATE_ARGUMENT(EPopupMethod, Method)
 		SLATE_ARGUMENT(int32, MaxChatLength)
-		SLATE_ARGUMENT(bool, bInGameUI)
 	SLATE_END_ARGS()
 
 	/**
diff --git a/Engine/Source/Developer/FriendsAndChat/Public/Interfaces/IFriendsAndChatManager.h b/Engine/Source/Developer/FriendsAndChat/Public/Interfaces/IFriendsAndChatManager.h
--- a/Engine/Source/Developer/FriendsAndChat/Public/Interfaces/IFriendsAndChatManager.h
+++ b/Engine/Source/Developer/FriendsAndChat/Public/Interfaces/IFriendsAndChatManager.h
@@ -41,9 +41,10 @@ public:
 	/**
 	 * Generate a chat widget.
 	 * @param InStyle The style to use to create the widgets.
+	 * @param The chat view model.
 	 * @return The chat widget.
 	 */
-	virtual TSharedPtr< SWidget > GenerateChatWidget(const FFriendsAndChatStyle* InStyle, bool bInGameUI) = 0;
+	virtual TSharedPtr< SWidget > GenerateChatWidget(const FFriendsAndChatStyle* InStyle, TSharedRef<IChatViewModel> ViewModel) = 0;
 
 	/**
 	 * Get the chat system view model for manipulating the chat widget.
