#FriendsAndChat - Refactor create friends list function to be more descriptive. Update Friends and chat to query join game option

diff --git a/Engine/Source/Developer/FriendsAndChat/Private/FriendsAndChatManager.cpp b/Engine/Source/Developer/FriendsAndChat/Private/FriendsAndChatManager.cpp
--- a/Engine/Source/Developer/FriendsAndChat/Private/FriendsAndChatManager.cpp
+++ b/Engine/Source/Developer/FriendsAndChat/Private/FriendsAndChatManager.cpp
@@ -233,7 +233,7 @@ void FFriendsAndChatManager::JoinPublicChatRoom(const FString& RoomName)
 
 // UI Creation
 
-void FFriendsAndChatManager::CreateFriendsListWidget(const FFriendsAndChatStyle* InStyle )
+void FFriendsAndChatManager::CreateFriendsListWindow(const FFriendsAndChatStyle* InStyle )
 {
 	const FVector2D DEFAULT_WINDOW_SIZE = FVector2D(308, 458);
 
diff --git a/Engine/Source/Developer/FriendsAndChat/Private/FriendsAndChatManager.h b/Engine/Source/Developer/FriendsAndChat/Private/FriendsAndChatManager.h
--- a/Engine/Source/Developer/FriendsAndChat/Private/FriendsAndChatManager.h
+++ b/Engine/Source/Developer/FriendsAndChat/Private/FriendsAndChatManager.h
@@ -125,7 +125,7 @@ public:
 	virtual void Login() override;
 	virtual bool IsLoggedIn() override;
 	virtual void SetApplicationViewModel(TSharedPtr<IFriendsApplicationViewModel> ApplicationViewModel) override;
-	virtual void CreateFriendsListWidget(const FFriendsAndChatStyle* InStyle ) override;
+	virtual void CreateFriendsListWindow(const FFriendsAndChatStyle* InStyle ) override;
 	virtual void SetUserSettings(const FFriendsAndChatSettings& UserSettings) override;
 	virtual void SetAnalyticsProvider(const TSharedPtr<IAnalyticsProvider>& AnalyticsProvider) override;
 	virtual TSharedPtr< SWidget > GenerateFriendsListWidget( const FFriendsAndChatStyle* InStyle ) override;
diff --git a/Engine/Source/Developer/FriendsAndChat/Private/UI/Widgets/SFriendItem.cpp b/Engine/Source/Developer/FriendsAndChat/Private/UI/Widgets/SFriendItem.cpp
--- a/Engine/Source/Developer/FriendsAndChat/Private/UI/Widgets/SFriendItem.cpp
+++ b/Engine/Source/Developer/FriendsAndChat/Private/UI/Widgets/SFriendItem.cpp
@@ -153,6 +153,7 @@ private:
 			.Padding(5)
 			[
 				SNew(SButton)
+				.IsEnabled(this, &SFriendItemImpl::IsActionEnabled, FriendAction)
 				.OnClicked(this, &SFriendItemImpl::HandleActionClicked, FriendAction)
 				.ButtonStyle(&FriendStyle.FriendListActionButtonStyle)
 				.VAlign(VAlign_Center)
@@ -300,6 +301,11 @@ private:
 		return FReply::Handled();
 	}
 
+	bool IsActionEnabled(EFriendActionType::Type FriendAction) const
+	{
+		return FriendAction == EFriendActionType::JoinGame ? ViewModel->CanJoinGame() : true;
+	}
+
 	FReply HandleRemoveClicked(bool bConfirm)
 	{
 		RemoveConfirmationAnchor->SetIsOpen(false);
diff --git a/Engine/Source/Developer/FriendsAndChat/Public/Interfaces/IFriendsAndChatManager.h b/Engine/Source/Developer/FriendsAndChat/Public/Interfaces/IFriendsAndChatManager.h
--- a/Engine/Source/Developer/FriendsAndChat/Public/Interfaces/IFriendsAndChatManager.h
+++ b/Engine/Source/Developer/FriendsAndChat/Public/Interfaces/IFriendsAndChatManager.h
@@ -10,12 +10,10 @@ class IFriendsAndChatManager
 public:
 
 	/**
-	 * Create the friends list widget.
-	 *
-	 * @param ParentWindow The parent window to add the widget
+	 * Create a friends list window.
 	 * @param InStyle The style to use to create the widgets
 	 */
-	virtual void CreateFriendsListWidget(const struct FFriendsAndChatStyle* InStyle ) = 0;
+	virtual void CreateFriendsListWindow(const struct FFriendsAndChatStyle* InStyle ) = 0;
 
 	/**
 	 * Set the FriendsAndChatUserSettings.
