#FriendsAndChat - Fix up Join game option

diff --git a/Engine/Source/Developer/FriendsAndChat/Private/FriendsAndChatManager.cpp b/Engine/Source/Developer/FriendsAndChat/Private/FriendsAndChatManager.cpp
--- a/Engine/Source/Developer/FriendsAndChat/Private/FriendsAndChatManager.cpp
+++ b/Engine/Source/Developer/FriendsAndChat/Private/FriendsAndChatManager.cpp
@@ -626,6 +626,10 @@ bool FFriendsAndChatManager::JoinGameAllowed()
 	{
 		return AllowFriendsJoinGame().Execute();
 	}
+	else if (ApplicationViewModel.IsValid())
+	{
+		return ApplicationViewModel->IsAppJoinable();
+	}
 	return false;
 }
 
@@ -1355,11 +1359,6 @@ void FFriendsAndChatManager::AcceptGameInvite(const TSharedPtr<IFriendItem>& Fri
 	Analytics.RecordGameInvite(*FriendItem->GetUniqueID(), TEXT("Social.GameInvite.Accept"));
 }
 
-const bool FFriendsAndChatManager::CanJoinGame() const
-{
-	return !ApplicationViewModel.IsValid() || ApplicationViewModel->IsAppJoinable();
-}
-
 void FFriendsAndChatManager::SendGameInvite(const TSharedPtr<IFriendItem>& FriendItem)
 {
 	SendGameInvite(*FriendItem->GetUniqueID());
diff --git a/Engine/Source/Developer/FriendsAndChat/Private/FriendsAndChatManager.h b/Engine/Source/Developer/FriendsAndChat/Private/FriendsAndChatManager.h
--- a/Engine/Source/Developer/FriendsAndChat/Private/FriendsAndChatManager.h
+++ b/Engine/Source/Developer/FriendsAndChat/Private/FriendsAndChatManager.h
@@ -258,12 +258,6 @@ public:
 	 */
 	void AcceptGameInvite(const TSharedPtr<IFriendItem>& FriendItem);
 
-	/**
-	 * Can we join a game
-	 * @return True if able to join a game
-	 */
-	const bool CanJoinGame() const;
-
 	/**
 	 * Send a game invite to a friend
 	 *
diff --git a/Engine/Source/Developer/FriendsAndChat/Private/Models/FriendViewModel.cpp b/Engine/Source/Developer/FriendsAndChat/Private/Models/FriendViewModel.cpp
--- a/Engine/Source/Developer/FriendsAndChat/Private/Models/FriendViewModel.cpp
+++ b/Engine/Source/Developer/FriendsAndChat/Private/Models/FriendViewModel.cpp
@@ -174,11 +174,6 @@ public:
 		return FriendItem.IsValid() ? FriendItem->IsOnline() : false;
 	}
 
-	virtual bool CanJoinGame() const override
-	{
-		return FFriendsAndChatManager::Get()->CanJoinGame();
-	}
-
 	virtual EOnlinePresenceState::Type GetOnlineStatus() const override
 	{
 		return FriendItem.IsValid() ? FriendItem->GetOnlineStatus() : EOnlinePresenceState::Offline;
diff --git a/Engine/Source/Developer/FriendsAndChat/Private/Models/FriendViewModel.h b/Engine/Source/Developer/FriendsAndChat/Private/Models/FriendViewModel.h
--- a/Engine/Source/Developer/FriendsAndChat/Private/Models/FriendViewModel.h
+++ b/Engine/Source/Developer/FriendsAndChat/Private/Models/FriendViewModel.h
@@ -56,7 +56,6 @@ public:
 	virtual FText GetFriendLocation() const = 0;
 	virtual FString GetClientId() const = 0;
 	virtual bool IsOnline() const = 0;
-	virtual bool CanJoinGame() const = 0;
 	virtual EOnlinePresenceState::Type GetOnlineStatus() const = 0;
 };
 
diff --git a/Engine/Source/Developer/FriendsAndChat/Private/UI/Widgets/SInviteItem.cpp b/Engine/Source/Developer/FriendsAndChat/Private/UI/Widgets/SInviteItem.cpp
--- a/Engine/Source/Developer/FriendsAndChat/Private/UI/Widgets/SInviteItem.cpp
+++ b/Engine/Source/Developer/FriendsAndChat/Private/UI/Widgets/SInviteItem.cpp
@@ -91,7 +91,7 @@ private:
 
 	bool IsActionEnabled(EFriendActionType::Type FriendAction) const
 	{
-		return FriendAction == EFriendActionType::JoinGame ? ViewModel->CanJoinGame() : true;
+		return ViewModel->CanPerformAction(FriendAction);
 	}
 
 private:
