Launcher Friends and chat style changes

#proj Launcher

#change added header to chat
#change fixed corruption on buttons, was just useing a programmer art texture (OPP-1507)
#change fixed cut off text (OPP-1479)
#change Fixed white out in chat (OPP-1642)

diff --git a/Engine/Source/Developer/FriendsAndChat/Private/FriendsAndChatManager.cpp b/Engine/Source/Developer/FriendsAndChat/Private/FriendsAndChatManager.cpp
--- a/Engine/Source/Developer/FriendsAndChat/Private/FriendsAndChatManager.cpp
+++ b/Engine/Source/Developer/FriendsAndChat/Private/FriendsAndChatManager.cpp
@@ -5,9 +5,11 @@
 #include "SChatWindow.h"
 #include "FriendsViewModel.h"
 #include "FriendViewModel.h"
+#include "FriendsStatusViewModel.h"
 #include "ChatViewModel.h"
 #include "SNotificationList.h"
 #include "SWindowTitleBar.h"
+#include "SFriendsStatus.h"
 #include "FriendRecentPlayerItems.h"
 #include "FriendGameInviteItem.h"
 
@@ -235,7 +237,7 @@ void FFriendsAndChatManager::JoinPublicChatRoom(const FString& RoomName)
 
 void FFriendsAndChatManager::CreateFriendsListWindow(const FFriendsAndChatStyle* InStyle )
 {
-	const FVector2D DEFAULT_WINDOW_SIZE = FVector2D(308, 458);
+	const FVector2D DEFAULT_WINDOW_SIZE = FVector2D(325, 458);
 
 	Style = *InStyle;
 
@@ -416,6 +418,8 @@ void FFriendsAndChatManager::SetChatWindowContents()
 	ChatViewModel->SetInGameUI(false);
 	ChatViewModel->SetCaptureFocus(false);
 
+	TSharedRef< FFriendsStatusViewModel > StatusViewModel = FFriendsStatusViewModelFactory::Create(SharedThis(this));
+
 	ChatWindow->SetContent(
 		SNew( SBorder )
 		.VAlign( VAlign_Fill )
@@ -432,6 +436,15 @@ void FFriendsAndChatManager::SetChatWindowContents()
 				.Title(FText::GetEmpty())
 			]
 			+ SVerticalBox::Slot()
+			.AutoHeight()
+			.Padding(20.0f, 0.0f, 0.0f, 40.0f)
+			.VAlign(VAlign_Top)
+			[
+				SNew(SFriendsStatus, StatusViewModel)
+				.FriendStyle(&Style)
+				.Method(EPopupMethod::UseCurrentWindow)
+			]
+			+ SVerticalBox::Slot()
 			[
 				SNew(SChatWindow, ChatViewModel.ToSharedRef())
 				.FriendStyle( &Style )
diff --git a/Engine/Source/Developer/FriendsAndChat/Private/UI/Widgets/SFriendsContainer.cpp b/Engine/Source/Developer/FriendsAndChat/Private/UI/Widgets/SFriendsContainer.cpp
--- a/Engine/Source/Developer/FriendsAndChat/Private/UI/Widgets/SFriendsContainer.cpp
+++ b/Engine/Source/Developer/FriendsAndChat/Private/UI/Widgets/SFriendsContainer.cpp
@@ -73,7 +73,7 @@ public:
 								+ SHorizontalBox::Slot()
 								[
 									SAssignNew(FriendNameTextBox, SEditableTextBox)
-									.HintText(LOCTEXT("AddFriendHint", "Add friend by account name or email"))
+									.HintText(LOCTEXT("AddFriendHint", "Add friend by account name"))
 									.Style(&FriendStyle.AddFriendEditableTextStyle)
 									.OnTextCommitted(this, &SFriendsContainerImpl::HandleFriendEntered)
 								]
