#FriendsAndChat - More style updates. Updating padding, and adding buttons

diff --git a/Engine/Source/Developer/FriendsAndChat/Private/FriendsAndChatManager.cpp b/Engine/Source/Developer/FriendsAndChat/Private/FriendsAndChatManager.cpp
--- a/Engine/Source/Developer/FriendsAndChat/Private/FriendsAndChatManager.cpp
+++ b/Engine/Source/Developer/FriendsAndChat/Private/FriendsAndChatManager.cpp
@@ -437,7 +437,7 @@ void FFriendsAndChatManager::SetChatWindowContents()
 			]
 			+ SVerticalBox::Slot()
 			.AutoHeight()
-			.Padding(20.0f, 0.0f, 0.0f, 40.0f)
+			.Padding(20.0f, 0.0f)
 			.VAlign(VAlign_Top)
 			[
 				SNew(SFriendsStatus, StatusViewModel)
diff --git a/Engine/Source/Developer/FriendsAndChat/Private/FriendsAndChatStyle.cpp b/Engine/Source/Developer/FriendsAndChat/Private/FriendsAndChatStyle.cpp
--- a/Engine/Source/Developer/FriendsAndChat/Private/FriendsAndChatStyle.cpp
+++ b/Engine/Source/Developer/FriendsAndChat/Private/FriendsAndChatStyle.cpp
@@ -278,6 +278,18 @@ FFriendsAndChatStyle& FFriendsAndChatStyle::SetButtonPadding(const FVector2D& Pa
 	return *this;
 }
 
+FFriendsAndChatStyle& FFriendsAndChatStyle::SetBorderPadding(const FMargin& Padding)
+{
+	BorderPadding = Padding;
+	return *this;
+}
+
+FFriendsAndChatStyle& FFriendsAndChatStyle::SetFriendsListWidth(const float InFriendsListWidth)
+{
+	FriendsListWidth = InFriendsListWidth;
+	return *this;
+}
+
 FFriendsAndChatStyle& FFriendsAndChatStyle::SetChatBackgroundBrush(const FSlateBrush& InChatBackgroundBrush)
 {
 	ChatBackgroundBrush = InChatBackgroundBrush;
diff --git a/Engine/Source/Developer/FriendsAndChat/Private/UI/Widgets/SFriendsContainer.cpp b/Engine/Source/Developer/FriendsAndChat/Private/UI/Widgets/SFriendsContainer.cpp
--- a/Engine/Source/Developer/FriendsAndChat/Private/UI/Widgets/SFriendsContainer.cpp
+++ b/Engine/Source/Developer/FriendsAndChat/Private/UI/Widgets/SFriendsContainer.cpp
@@ -37,68 +37,60 @@ public:
 			.HAlign(HAlign_Fill)
 			[
 				SNew(SBorder)
-				.Padding(20.0f)
+				.Padding(FriendStyle.BorderPadding)
 				.BorderImage(&FriendStyle.FriendContainerHeader)
 				[
-					SNew(SHorizontalBox)
-					+SHorizontalBox::Slot()
-					.HAlign(HAlign_Left)
-					.VAlign(VAlign_Center)
-					.AutoWidth()
+					SNew(SBox)
+					.WidthOverride(FriendStyle.FriendsListWidth)
 					[
-						SNew(SFriendsStatus, ViewModel->GetStatusViewModel())
-						.FriendStyle(&FriendStyle)
-						.Method(MenuMethod)
-					]
-					+ SHorizontalBox::Slot()
-					.VAlign(VAlign_Center)
-					.HAlign(HAlign_Fill)
-					.Padding(FMargin(10,0,0,0))
-					[
-						SNew(SOverlay)
-						+ SOverlay::Slot()
+						SNew(SHorizontalBox)
+						+SHorizontalBox::Slot()
+						.VAlign(VAlign_Center)
+						.HAlign(HAlign_Left)
+						.AutoWidth()
 						[
-							SNew(SSpacer).Size(FVector2D(400.0f, 0.0f))
+							SNew(SFriendsStatus, ViewModel->GetStatusViewModel())
+							.FriendStyle(&FriendStyle)
+							.Method(MenuMethod)
 						]
-						+ SOverlay::Slot()
+						+ SHorizontalBox::Slot()
 						.VAlign(VAlign_Center)
 						.HAlign(HAlign_Fill)
 						[
 							SNew(SBorder)
-							.Padding(3.0f)
 							.Visibility(this, &SFriendsContainerImpl::AddFriendVisibility)
 							.BorderImage(&FriendStyle.AddFriendEditBorder)
+							.VAlign(VAlign_Center)
+							.HAlign(HAlign_Fill)
 							[
-								SNew(SHorizontalBox)
-								+ SHorizontalBox::Slot()
-								[
-									SAssignNew(FriendNameTextBox, SEditableTextBox)
-									.HintText(LOCTEXT("AddFriendHint", "Add friend by account name"))
-									.Style(&FriendStyle.AddFriendEditableTextStyle)
-									.OnTextCommitted(this, &SFriendsContainerImpl::HandleFriendEntered)
-								]
-								+ SHorizontalBox::Slot()
-									.AutoWidth()
-									[
-										SNew(SButton)
-										.ButtonStyle(&FriendStyle.AddFriendCloseButtonStyle)
-										.OnClicked(this, &SFriendsContainerImpl::HandleAddFriendButtonClicked)
-									]
+								SAssignNew(FriendNameTextBox, SEditableTextBox)
+								.HintText(LOCTEXT("AddFriendHint", "Add friend by account name or email"))
+								.Style(&FriendStyle.AddFriendEditableTextStyle)
+								.OnTextCommitted(this, &SFriendsContainerImpl::HandleFriendEntered)
 							]
 						]
-						+ SOverlay::Slot()
+						+SHorizontalBox::Slot()
 						.VAlign(VAlign_Center)
 						.HAlign(HAlign_Right)
+						.Padding(4)
 						[
 							SNew(SButton)
-							.ContentPadding(FMargin(0.0f, 5.0f))
-							.ButtonStyle(&FriendStyle.FriendListActionButtonStyle)
+							.ButtonStyle(&FriendStyle.AddFriendButtonStyle)
 							.OnClicked(this, &SFriendsContainerImpl::HandleAddFriendButtonClicked)
 							.Visibility(this, &SFriendsContainerImpl::AddFriendActionVisibility)
-							[
-								SNew(SImage)
-								.Image(&FriendStyle.FriendsAddImageBrush)
-							]
+							.Cursor(EMouseCursor::Hand)
+						]
+						+ SHorizontalBox::Slot()
+						.VAlign(VAlign_Center)
+						.HAlign(HAlign_Right)
+						.AutoWidth()
+						.Padding(4)
+						[
+							SNew(SButton)
+							.ButtonStyle(&FriendStyle.AddFriendCloseButtonStyle)
+							.OnClicked(this, &SFriendsContainerImpl::HandleAddFriendButtonClicked)
+							.Visibility(this, &SFriendsContainerImpl::AddFriendVisibility)
+							.Cursor(EMouseCursor::Hand)
 						]
 					]
 				]
@@ -106,7 +98,7 @@ public:
 			+ SVerticalBox::Slot()
 			[
 				SNew(SBorder)
-				.Padding(20.0f)
+				.Padding(FriendStyle.BorderPadding)
 				.BorderImage(&FriendStyle.FriendContainerBackground)
 				[
 					SNew(SScrollBox)
diff --git a/Engine/Source/Developer/FriendsAndChat/Private/UI/Widgets/SFriendsListContainer.cpp b/Engine/Source/Developer/FriendsAndChat/Private/UI/Widgets/SFriendsListContainer.cpp
--- a/Engine/Source/Developer/FriendsAndChat/Private/UI/Widgets/SFriendsListContainer.cpp
+++ b/Engine/Source/Developer/FriendsAndChat/Private/UI/Widgets/SFriendsListContainer.cpp
@@ -33,7 +33,7 @@ public:
 			.AutoHeight()
 			[
 				SNew(SBorder)
-				.Padding(15.0f)
+				.Padding(FriendStyle.BorderPadding)
 				.BorderImage(&FriendStyle.FriendListHeader)
 				[
 					SNew(SButton)
diff --git a/Engine/Source/Developer/FriendsAndChat/Public/FriendsAndChatStyle.h b/Engine/Source/Developer/FriendsAndChat/Public/FriendsAndChatStyle.h
--- a/Engine/Source/Developer/FriendsAndChat/Public/FriendsAndChatStyle.h
+++ b/Engine/Source/Developer/FriendsAndChat/Public/FriendsAndChatStyle.h
@@ -252,6 +252,14 @@ struct FRIENDSANDCHAT_API FFriendsAndChatStyle
 	FVector2D ButtonPadding;
 	FFriendsAndChatStyle& SetButtonPadding(const FVector2D& ButtonPadding);
 
+	UPROPERTY(EditAnywhere, Category = Appearance)
+	FMargin BorderPadding;
+	FFriendsAndChatStyle& SetBorderPadding(const FMargin& BorderPadding);
+
+	UPROPERTY(EditAnywhere, Category = Appearance)
+	float FriendsListWidth;
+	FFriendsAndChatStyle& SetFriendsListWidth(const float FriendsListLength);
+
 	UPROPERTY(EditAnywhere, Category = Appearance)
 	FSlateBrush ChatBackgroundBrush;
 	FFriendsAndChatStyle& SetChatBackgroundBrush(const FSlateBrush& InChatBackgroundBrush);
