#FriendsAndChat - Updates to message storage, style updates, and removing Network message type

diff --git a/Engine/Source/Developer/FriendsAndChat/Private/FriendsAndChatManager.cpp b/Engine/Source/Developer/FriendsAndChat/Private/FriendsAndChatManager.cpp
--- a/Engine/Source/Developer/FriendsAndChat/Private/FriendsAndChatManager.cpp
+++ b/Engine/Source/Developer/FriendsAndChat/Private/FriendsAndChatManager.cpp
@@ -237,7 +237,7 @@ void FFriendsAndChatManager::JoinPublicChatRoom(const FString& RoomName)
 
 void FFriendsAndChatManager::CreateFriendsListWindow(const FFriendsAndChatStyle* InStyle )
 {
-	const FVector2D DEFAULT_WINDOW_SIZE = FVector2D(325, 458);
+	const FVector2D DEFAULT_WINDOW_SIZE = FVector2D(350, 458);
 
 	Style = *InStyle;
 
diff --git a/Engine/Source/Developer/FriendsAndChat/Private/FriendsMessageManager.cpp b/Engine/Source/Developer/FriendsAndChat/Private/FriendsMessageManager.cpp
--- a/Engine/Source/Developer/FriendsAndChat/Private/FriendsMessageManager.cpp
+++ b/Engine/Source/Developer/FriendsAndChat/Private/FriendsMessageManager.cpp
@@ -4,6 +4,11 @@
 #include "OnlineChatInterface.h"
 
 #define LOCTEXT_NAMESPACE "FriendsMessageManager"
+// Message expiry time for different message types
+static const int32 GlobalMessageLifetime = 5 * 60.0f;  // 5 min
+static const int32 PartyMessageLifetime = 5 * 60.0f;  // 5 min
+static const int32 WhisperMessageLifetime = 5 * 60.0f;  // 5 min
+static const int32 MessageStore = 200;
 
 class FFriendsMessageManagerImpl
 	: public FFriendsMessageManager
@@ -109,6 +114,7 @@ public:
 		ChatItem->Message = FText::FromString(MsgBody);
 		ChatItem->MessageType = EChatMessageType::Party;
 		ChatItem->MessageTimeText = FText::AsTime(FDateTime::UtcNow());
+		ChatItem->ExpireTime = PartyMessageLifetime;
 		ChatItem->bIsFromSelf = false;
 		AddMessage(ChatItem.ToSharedRef());
 	}
@@ -241,6 +247,7 @@ private:
 			ChatItem->Message = FText::FromString(TEXT("entered room"));
 			ChatItem->MessageType = EChatMessageType::Global;
 			ChatItem->MessageTimeText = FText::AsTime(FDateTime::UtcNow());
+			ChatItem->ExpireTime = FDateTime::UtcNow() + GlobalMessageLifetime;
 			ChatItem->bIsFromSelf = false;
 			AddMessage(ChatItem);
 		}
@@ -261,6 +268,7 @@ private:
 			ChatItem->Message = FText::FromString(TEXT("left room"));
 			ChatItem->MessageType = EChatMessageType::Global;
 			ChatItem->MessageTimeText = FText::AsTime(FDateTime::UtcNow());
+			ChatItem->ExpireTime = FDateTime::UtcNow() + GlobalMessageLifetime;
 			ChatItem->bIsFromSelf = false;
 			AddMessage(ChatItem);
 		}
@@ -278,6 +286,7 @@ private:
 		ChatItem->Message = FText::FromString(*ChatMessage->GetBody());
 		ChatItem->MessageType = EChatMessageType::Global;
 		ChatItem->MessageTimeText = FText::AsTime(ChatMessage->GetTimestamp());
+		ChatItem->ExpireTime = ChatMessage->GetTimestamp() + GlobalMessageLifetime;
 		ChatItem->bIsFromSelf = ChatMessage->GetUserId() == *LoggedInUser;
 		TSharedPtr<IFriendItem> FoundFriend = FFriendsAndChatManager::Get()->FindUser(ChatMessage->GetUserId());
 		if(FoundFriend.IsValid())
@@ -303,6 +312,7 @@ private:
 			ChatItem->Message = FText::FromString(*ChatMessage->GetBody());
 			ChatItem->MessageType = EChatMessageType::Whisper;
 			ChatItem->MessageTimeText = FText::AsTime(ChatMessage->GetTimestamp());
+			ChatItem->ExpireTime = ChatMessage->GetTimestamp() + WhisperMessageLifetime;
 			ChatItem->bIsFromSelf = false;
 			ChatItem->MessageRef = ChatMessage;
 			AddMessage(ChatItem);
@@ -326,8 +336,41 @@ private:
 
 	void AddMessage(TSharedPtr< FFriendChatMessage > ChatItem)
 	{
-		ReceivedMessages.Add(ChatItem);
-		OnChatMessageRecieved().Broadcast(ChatItem.ToSharedRef());
+		if(ReceivedMessages.Add(ChatItem) > MessageStore)
+		{
+			bool bGlobalTimeFound = false;
+			bool bPartyTimeFound = false;
+			bool bWhisperFound = false;
+			FDateTime CurrentTime = FDateTime::UtcNow();
+			for(int32 Index = 0; Index < ReceivedMessages.Num(); Index++)
+			{
+				TSharedPtr<FFriendChatMessage> Message = ReceivedMessages[Index];
+				if(Message->ExpireTime < CurrentTime)
+				{
+					ReceivedMessages.Remove(Message);
+					Index--;
+				}
+				else
+				{
+					switch(Message->MessageType)
+					{
+						case EChatMessageType::Global : bGlobalTimeFound = true; break;
+						case EChatMessageType::Party : bPartyTimeFound = true; break;
+						case EChatMessageType::Whisper : bWhisperFound = true; break;
+					}
+				}
+
+				if(ReceivedMessages.Num() < MessageStore || (bPartyTimeFound && bGlobalTimeFound && bWhisperFound))
+				{
+					break;
+				}
+			}
+			OnChatMessageRecieved().Broadcast(ChatItem.ToSharedRef());
+		}
+		else
+		{
+			OnChatMessageRecieved().Broadcast(ChatItem.ToSharedRef());
+		}
 	}
 
 private:
@@ -351,7 +394,11 @@ private:
 	TSharedPtr<FUniqueNetId> LoggedInUser;
 	TArray<FString> RoomJoins;
 
-	TArray<TSharedPtr< FFriendChatMessage > > ReceivedMessages; 
+	TArray<TSharedPtr< FFriendChatMessage > > ReceivedMessages;
+
+	TArray<TSharedPtr< FFriendChatMessage > > GlobalMessages;
+	TArray<TSharedPtr< FFriendChatMessage > > WhisperMessages;
+	TArray<TSharedPtr< FFriendChatMessage > > PartyMessages;
 
 	bool bEnableEnterExitMessages;
 
diff --git a/Engine/Source/Developer/FriendsAndChat/Private/FriendsMessageManager.h b/Engine/Source/Developer/FriendsAndChat/Private/FriendsMessageManager.h
--- a/Engine/Source/Developer/FriendsAndChat/Private/FriendsMessageManager.h
+++ b/Engine/Source/Developer/FriendsAndChat/Private/FriendsMessageManager.h
@@ -15,8 +15,6 @@ namespace EChatMessageType
 		Party,
 		// Global Chat Item
 		Global,
-		// Network chat item
-		Network,
 	};
 
 	/** @return the FTextified version of the enum passed in */
@@ -27,7 +25,6 @@ namespace EChatMessageType
 			case Global: return NSLOCTEXT("FriendsList","Global", "Global");
 			case Whisper: return NSLOCTEXT("FriendsList","Whisper", "Whisper");
 			case Party: return NSLOCTEXT("FriendsList","Party", "Party");
-			case Network: return NSLOCTEXT("FriendsList","Network", "Network");
 
 			default: return FText::GetEmpty();
 		}
@@ -41,6 +38,7 @@ struct FFriendChatMessage
 	FText FromName;
 	FText Message;
 	FText MessageTimeText;
+	FDateTime ExpireTime;
 	TSharedPtr<class FChatMessage> MessageRef;
 	TSharedPtr<FUniqueNetId> SenderId;
 	bool bIsFromSelf;
diff --git a/Engine/Source/Developer/FriendsAndChat/Private/Models/ChatViewModel.cpp b/Engine/Source/Developer/FriendsAndChat/Private/Models/ChatViewModel.cpp
--- a/Engine/Source/Developer/FriendsAndChat/Private/Models/ChatViewModel.cpp
+++ b/Engine/Source/Developer/FriendsAndChat/Private/Models/ChatViewModel.cpp
@@ -5,6 +5,9 @@
 #include "ChatViewModel.h"
 #include "OnlineChatInterface.h"
 
+// Holds how many message to display
+static const int32 MessageDisplay = 200;
+
 class FChatViewModelImpl
 	: public FChatViewModel
 {
@@ -401,21 +404,27 @@ private:
 
 	void FilterChatList()
 	{
+		bool bAddedItem = true;
 		if(!IsGlobalChatEnabled())
 		{
 			FilteredChatLists.Empty();
+			bAddedItem = false;
 			for (const auto& ChatItem : ChatLists)
 			{
 				if(ChatItem->GetMessageType() != EChatMessageType::Global)
 				{
 					FilteredChatLists.Add(ChatItem);
-					ChatListUpdatedEvent.Broadcast();
+					bAddedItem = true;
 				}
 			}
 		}
 		else
 		{
 			FilteredChatLists = ChatLists;
+		}
+
+		if(bAddedItem)
+		{
 			ChatListUpdatedEvent.Broadcast();
 		}
 	}
@@ -435,7 +444,14 @@ private:
 
 	void HandleMessageReceived(const TSharedRef<FFriendChatMessage> NewMessage)
 	{
-		ChatLists.Add(FChatItemViewModelFactory::Create(NewMessage, SharedThis(this)));
+		if(IsGlobalChatEnabled() || NewMessage->MessageType != EChatMessageType::Global)
+		{
+			ChatLists.Add(FChatItemViewModelFactory::Create(NewMessage, SharedThis(this)));
+			if(ChatLists.Num() > MessageDisplay)
+			{
+				ChatLists.RemoveAt(0);
+			}
+		}
 		FilterChatList();
 	}
 
diff --git a/Engine/Source/Developer/FriendsAndChat/Private/UI/Widgets/SChatItem.cpp b/Engine/Source/Developer/FriendsAndChat/Private/UI/Widgets/SChatItem.cpp
--- a/Engine/Source/Developer/FriendsAndChat/Private/UI/Widgets/SChatItem.cpp
+++ b/Engine/Source/Developer/FriendsAndChat/Private/UI/Widgets/SChatItem.cpp
@@ -108,7 +108,6 @@ private:
 				case EChatMessageType::Global: return FriendStyle.DefaultChatColor.CopyWithNewOpacity(ViewModel->GetFadeAmountColor()); break;
 				case EChatMessageType::Whisper: return FriendStyle.WhisplerChatColor.CopyWithNewOpacity(ViewModel->GetFadeAmountColor()); break;
 				case EChatMessageType::Party: return FriendStyle.PartyChatColor.CopyWithNewOpacity(ViewModel->GetFadeAmountColor()); break;
-				case EChatMessageType::Network: return FriendStyle.NetworkChatColor.CopyWithNewOpacity(ViewModel->GetFadeAmountColor()); break;
 				default: return FLinearColor::Gray;
 			}
 		}
@@ -128,7 +127,6 @@ private:
 				case EChatMessageType::Global: DisplayColor = FriendStyle.DefaultChatColor; break;
 				case EChatMessageType::Whisper: DisplayColor =  FriendStyle.WhisplerChatColor; break;
 				case EChatMessageType::Party: DisplayColor =  FriendStyle.PartyChatColor; break;
-				case EChatMessageType::Network: DisplayColor = FriendStyle.NetworkChatColor; break;
 				default: DisplayColor = FLinearColor::Gray;
 			}
 			return DisplayColor;
diff --git a/Engine/Source/Developer/FriendsAndChat/Private/UI/Widgets/SChatWindow.cpp b/Engine/Source/Developer/FriendsAndChat/Private/UI/Widgets/SChatWindow.cpp
--- a/Engine/Source/Developer/FriendsAndChat/Private/UI/Widgets/SChatWindow.cpp
+++ b/Engine/Source/Developer/FriendsAndChat/Private/UI/Widgets/SChatWindow.cpp
@@ -384,7 +384,6 @@ private:
 				case EChatMessageType::Global: ChannelColor = FriendStyle.DefaultChatColor; break;
 				case EChatMessageType::Whisper: ChannelColor = FriendStyle.WhisplerChatColor; break;
 				case EChatMessageType::Party: ChannelColor = FriendStyle.PartyChatColor; break;
-				case EChatMessageType::Network: ChannelColor = FriendStyle.NetworkChatColor; break;
 			}
 
 			ChannelSelection->AddSlot()
