#FriendsAndChat - Setting the message exiry time on party messages correctly

diff --git a/Engine/Source/Developer/FriendsAndChat/Private/FriendsMessageManager.cpp b/Engine/Source/Developer/FriendsAndChat/Private/FriendsMessageManager.cpp
--- a/Engine/Source/Developer/FriendsAndChat/Private/FriendsMessageManager.cpp
+++ b/Engine/Source/Developer/FriendsAndChat/Private/FriendsMessageManager.cpp
@@ -114,7 +114,7 @@ public:
 		ChatItem->Message = FText::FromString(MsgBody);
 		ChatItem->MessageType = EChatMessageType::Party;
 		ChatItem->MessageTimeText = FText::AsTime(FDateTime::UtcNow());
-		ChatItem->ExpireTime = PartyMessageLifetime;
+		ChatItem->ExpireTime = FDateTime::UtcNow() + PartyMessageLifetime;
 		ChatItem->bIsFromSelf = false;
 		AddMessage(ChatItem.ToSharedRef());
 	}
