#FriendsAndChat - Update to chat message display. Convert to Rich Text

diff --git a/Engine/Source/Developer/FriendsAndChat/Private/FriendsAndChatManager.cpp b/Engine/Source/Developer/FriendsAndChat/Private/FriendsAndChatManager.cpp
--- a/Engine/Source/Developer/FriendsAndChat/Private/FriendsAndChatManager.cpp
+++ b/Engine/Source/Developer/FriendsAndChat/Private/FriendsAndChatManager.cpp
@@ -240,6 +240,7 @@ void FFriendsAndChatManager::CreateFriendsListWindow(const FFriendsAndChatStyle*
 	const FVector2D DEFAULT_WINDOW_SIZE = FVector2D(350, 458);
 
 	Style = *InStyle;
+	FFriendsAndChatModuleStyle::Initialize(Style);
 
 	if (!FriendWindow.IsValid())
 	{
@@ -323,6 +324,7 @@ TSharedPtr< SWidget > FFriendsAndChatManager::GenerateFriendsListWidget( const F
 	if ( !FriendListWidget.IsValid() )
 	{
 		Style = *InStyle;
+		FFriendsAndChatModuleStyle::Initialize(Style);
 		SAssignNew(FriendListWidget, SOverlay)
 		+SOverlay::Slot()
 		[
diff --git a/Engine/Source/Developer/FriendsAndChat/Private/FriendsAndChatStyle.cpp b/Engine/Source/Developer/FriendsAndChat/Private/FriendsAndChatStyle.cpp
--- a/Engine/Source/Developer/FriendsAndChat/Private/FriendsAndChatStyle.cpp
+++ b/Engine/Source/Developer/FriendsAndChat/Private/FriendsAndChatStyle.cpp
@@ -302,6 +302,12 @@ FFriendsAndChatStyle& FFriendsAndChatStyle::SetFriendsListWidth(const float InFr
 	return *this;
 }
 
+FFriendsAndChatStyle& FFriendsAndChatStyle::SetChatListWidth(const float InChatListWidth)
+{
+	ChatListWidth = InChatListWidth;
+	return *this;
+}
+
 FFriendsAndChatStyle& FFriendsAndChatStyle::SetChatBackgroundBrush(const FSlateBrush& InChatBackgroundBrush)
 {
 	ChatBackgroundBrush = InChatBackgroundBrush;
@@ -326,3 +332,89 @@ const FFriendsAndChatStyle& FFriendsAndChatStyle::GetDefault()
 	return Default;
 }
 
+/**
+	Module style set
+*/
+TSharedPtr< FSlateStyleSet > FFriendsAndChatModuleStyle::FriendsAndChatModuleStyleInstance = NULL;
+
+void FFriendsAndChatModuleStyle::Initialize(FFriendsAndChatStyle FriendStyle)
+{
+	if ( !FriendsAndChatModuleStyleInstance.IsValid() )
+	{
+		FriendsAndChatModuleStyleInstance = Create(FriendStyle);
+		FSlateStyleRegistry::RegisterSlateStyle( *FriendsAndChatModuleStyleInstance );
+	}
+}
+
+void FFriendsAndChatModuleStyle::Shutdown()
+{
+	FSlateStyleRegistry::UnRegisterSlateStyle( *FriendsAndChatModuleStyleInstance );
+	ensure( FriendsAndChatModuleStyleInstance.IsUnique() );
+	FriendsAndChatModuleStyleInstance.Reset();
+}
+
+FName FFriendsAndChatModuleStyle::GetStyleSetName()
+{
+	static FName StyleSetName(TEXT("FriendsAndChat"));
+	return StyleSetName;
+}
+
+TSharedRef< FSlateStyleSet > FFriendsAndChatModuleStyle::Create(FFriendsAndChatStyle FriendStyle)
+{
+	TSharedRef< FSlateStyleSet > Style = MakeShareable(new FSlateStyleSet("FriendsAndChatStyle"));
+
+	const FTextBlockStyle DefaultText = FTextBlockStyle()
+		.SetFont(FriendStyle.FriendsFontStyleSmall);
+
+	// Name Style
+	const FTextBlockStyle GlobalChatFont = FTextBlockStyle(DefaultText)
+		.SetFont(FriendStyle.FriendsFontStyleSmallBold)
+		.SetColorAndOpacity(FriendStyle.DefaultChatColor);
+
+	const FTextBlockStyle PartyChatFont = FTextBlockStyle(DefaultText)
+		.SetFont(FriendStyle.FriendsFontStyleSmallBold)
+		.SetColorAndOpacity(FriendStyle.PartyChatColor);
+
+	const FTextBlockStyle WhisperChatFont = FTextBlockStyle(DefaultText)
+		.SetFont(FriendStyle.FriendsFontStyleSmallBold)
+		.SetColorAndOpacity(FriendStyle.WhisplerChatColor);
+
+	const FButtonStyle UserNameButton = FButtonStyle()
+		.SetNormal(FSlateNoResource())
+		.SetPressed(FSlateNoResource())
+		.SetHovered(FSlateNoResource());
+
+	const FHyperlinkStyle GlobalChatHyperlink = FHyperlinkStyle()
+		.SetUnderlineStyle(UserNameButton)
+		.SetTextStyle(GlobalChatFont)
+		.SetPadding(FMargin(0.0f));
+
+	const FHyperlinkStyle PartyChatHyperlink = FHyperlinkStyle()
+		.SetUnderlineStyle(UserNameButton)
+		.SetTextStyle(PartyChatFont)
+		.SetPadding(FMargin(0.0f));
+
+	const FHyperlinkStyle WhisperChatHyperlink = FHyperlinkStyle()
+		.SetUnderlineStyle(UserNameButton)
+		.SetTextStyle(WhisperChatFont)
+		.SetPadding(FMargin(0.0f));
+
+	Style->Set("UserNameTextStyle.GlobalHyperlink", GlobalChatHyperlink);
+	Style->Set("UserNameTextStyle.PartyHyperlink", PartyChatHyperlink);
+	Style->Set("UserNameTextStyle.Whisperlink", WhisperChatHyperlink);
+	Style->Set("UserNameTextStyle.GlobalTextStyle", GlobalChatFont);
+	Style->Set("UserNameTextStyle.PartyTextStyle", PartyChatFont);
+	Style->Set("UserNameTextStyle.WhisperTextStyle", WhisperChatFont);
+
+	return Style;
+}
+
+void FFriendsAndChatModuleStyle::ReloadTextures()
+{
+	FSlateApplication::Get().GetRenderer()->ReloadTextureResources();
+}
+
+const ISlateStyle& FFriendsAndChatModuleStyle::Get()
+{
+	return *FriendsAndChatModuleStyleInstance;
+}
\ No newline at end of file
diff --git a/Engine/Source/Developer/FriendsAndChat/Private/Models/ChatItemViewModel.cpp b/Engine/Source/Developer/FriendsAndChat/Private/Models/ChatItemViewModel.cpp
--- a/Engine/Source/Developer/FriendsAndChat/Private/Models/ChatItemViewModel.cpp
+++ b/Engine/Source/Developer/FriendsAndChat/Private/Models/ChatItemViewModel.cpp
@@ -69,7 +69,14 @@ public:
 
 	virtual void FriendNameSelected() const override
 	{
-		Owner->SetChannelUserClicked(ChatMessage);
+		if(!IsFromSelf() || GetMessageType() == EChatMessageType::Whisper)
+		{
+			Owner->SetChannelUserClicked(ChatMessage);
+		}
+		else
+		{
+			Owner->SetChatChannel(GetMessageType());
+		}
 	}
 
 private:
diff --git a/Engine/Source/Developer/FriendsAndChat/Private/UI/Widgets/SChatItem.cpp b/Engine/Source/Developer/FriendsAndChat/Private/UI/Widgets/SChatItem.cpp
--- a/Engine/Source/Developer/FriendsAndChat/Private/UI/Widgets/SChatItem.cpp
+++ b/Engine/Source/Developer/FriendsAndChat/Private/UI/Widgets/SChatItem.cpp
@@ -27,6 +27,15 @@ public:
 			}
 		}
 
+		FFormatNamedArguments Args;
+		Args.Add(TEXT("DisplayName"), DisplayNameText);
+		Args.Add(TEXT("Message"), ViewModel->GetMessage());
+		Args.Add(TEXT("NameStyle"), FText::FromString(GetTextHyperlinkStyle()));
+		FText DisplayText = FText::Format(LOCTEXT("SChatItem_Message", "<a id=\"UserName\" style=\"{NameStyle}\">{DisplayName}</>{Message}"), Args);
+
+		FTextBlockStyle TextStyle = FriendStyle.TextStyle;
+		TextStyle.ColorAndOpacity = GetChannelColor();
+
 		SUserWidget::Construct(SUserWidget::FArguments()
 		[
 			SNew(SHorizontalBox)
@@ -42,45 +51,18 @@ public:
 			+SHorizontalBox::Slot()
 			.AutoWidth()
 			.VAlign(VAlign_Center)
-			.Padding(FMargin(5,1))
-			.MaxWidth(150)
-			[
-				SNew(STextBlock)
-				.Visibility(ViewModel->IsFromSelf() && ViewModel->GetMessageType() != EChatMessageType::Party ? EVisibility::Visible : EVisibility::Collapsed)
-				.Text(DisplayNameText)
-				.ColorAndOpacity(this, &SChatItemImpl::GetChannelColor)
-				.Font(FriendStyle.FriendsFontStyleSmallBold)
-			]
-			+SHorizontalBox::Slot()
-			.AutoWidth()
-			.VAlign(VAlign_Center)
+			.HAlign(HAlign_Fill)
 			.Padding(FMargin(5, 1, 2, 1))
-			.MaxWidth(150)
+			.MaxWidth(FriendStyle.ChatListWidth)
 			[
-				SNew(SButton)
-				.Visibility(ViewModel->IsFromSelf() || ViewModel->GetMessageType() == EChatMessageType::Party ? EVisibility::Collapsed : EVisibility::Visible)
-				.ButtonStyle(&FriendStyle.FriendListItemButtonStyle)
-				.OnClicked(this, &SChatItemImpl::HandleNameClicked)
-				.ContentPadding(0)
-				[
-					SNew(STextBlock)
-					.Text(DisplayNameText)
-					.ColorAndOpacity(this, &SChatItemImpl::GetChannelColor)
-					.Font(FriendStyle.FriendsFontStyleSmallBold)
-				]
+				SNew( SRichTextBlock )
+				.Text(DisplayText)
+				.TextStyle(&TextStyle)
+				.DecoratorStyleSet(&FFriendsAndChatModuleStyle::Get())
+				.WrapTextAt(FriendStyle.ChatListWidth - 5)
+				+ SRichTextBlock::HyperlinkDecorator(TEXT( "UserName" ), this, &SChatItemImpl::HandleNameClicked)
 			]
 			+SHorizontalBox::Slot()
-			.Padding(FMargin(2, 1, 5, 1))
-			.VAlign(VAlign_Center)
-			[
-				SNew(STextBlock)
-				.Text(ViewModel->GetMessage())
-				.ColorAndOpacity(this, &SChatItemImpl::GetChannelColor)
-				.Font(FriendStyle.FriendsFontStyleSmall)
-				.AutoWrapText(true)
-			]
-			+SHorizontalBox::Slot()
-			.AutoWidth()
 			.HAlign(HAlign_Right)
 			.Padding(FMargin(5,1))
 			.VAlign(VAlign_Center)
@@ -89,7 +71,6 @@ public:
 				.Text(ViewModel->GetMessageTime())
 				.ColorAndOpacity(this, &SChatItemImpl::GetTimeDisplayColor)
 				.Font(FriendStyle.FriendsFontStyleSmall)
-				.AutoWrapText(true)
 			]
 		]);
 	}
@@ -152,10 +133,21 @@ private:
 		}
 	}
 
-	FReply HandleNameClicked()
+	const FString GetTextHyperlinkStyle() const
+	{
+		switch(ViewModel->GetMessageType())
+		{
+			case EChatMessageType::Global: return TEXT("UserNameTextStyle.GlobalHyperlink"); break;
+			case EChatMessageType::Whisper: return TEXT("UserNameTextStyle.Whisperlink"); break;
+			case EChatMessageType::Party: return TEXT("UserNameTextStyle.PartyHyperlink"); break;
+			default:
+			return FString();
+		}
+	}
+
+	void HandleNameClicked( const FSlateHyperlinkRun::FMetadata& Metadata )
 	{
 		ViewModel->FriendNameSelected();
-		return FReply::Handled();
 	}
 
 private:
diff --git a/Engine/Source/Developer/FriendsAndChat/Private/UI/Widgets/SChatWindow.cpp b/Engine/Source/Developer/FriendsAndChat/Private/UI/Widgets/SChatWindow.cpp
--- a/Engine/Source/Developer/FriendsAndChat/Private/UI/Widgets/SChatWindow.cpp
+++ b/Engine/Source/Developer/FriendsAndChat/Private/UI/Widgets/SChatWindow.cpp
@@ -43,6 +43,7 @@ public:
 					SNew(SHorizontalBox)
 					+SHorizontalBox::Slot()
 			 		.AutoWidth()
+					.VAlign(VAlign_Fill)
 					[
 						SNew( SBorder )
 						.BorderBackgroundColor(this, &SChatWindowImpl::GetTimedFadeSlateColor)
@@ -52,6 +53,7 @@ public:
 						]
 					]
 					+SHorizontalBox::Slot()
+					.VAlign(VAlign_Fill)
 			 		[
 						SAssignNew(ChatList, SListView<TSharedRef<FChatItemViewModel>>)
 						.ListItemsSource(&ViewModel->GetFilteredChatList())
@@ -93,8 +95,8 @@ public:
 					]
 					+SHorizontalBox::Slot()
 					[
-						SNew(SOverlay)
-						+ SOverlay::Slot()
+						SNew(SVerticalBox)
+						+ SVerticalBox::Slot()
 						[
 							SNew(SHorizontalBox)
 							.Visibility(this, &SChatWindowImpl::GetChatEntryVisibility)
@@ -153,7 +155,7 @@ public:
 								.OnTextChanged(this, &SChatWindowImpl::OnChatTextChanged)
 							]
 						]
-						+SOverlay::Slot()
+						+SVerticalBox::Slot()
 						[
 							SNew(SBorder)
 							.Visibility(this, &SChatWindowImpl::GetConfirmationVisibility)
diff --git a/Engine/Source/Developer/FriendsAndChat/Public/FriendsAndChatStyle.h b/Engine/Source/Developer/FriendsAndChat/Public/FriendsAndChatStyle.h
--- a/Engine/Source/Developer/FriendsAndChat/Public/FriendsAndChatStyle.h
+++ b/Engine/Source/Developer/FriendsAndChat/Public/FriendsAndChatStyle.h
@@ -270,6 +270,10 @@ struct FRIENDSANDCHAT_API FFriendsAndChatStyle
 	float FriendsListWidth;
 	FFriendsAndChatStyle& SetFriendsListWidth(const float FriendsListLength);
 
+	UPROPERTY(EditAnywhere, Category = Appearance)
+	float ChatListWidth;
+	FFriendsAndChatStyle& SetChatListWidth(const float ChatListWidth);
+
 	UPROPERTY(EditAnywhere, Category = Appearance)
 	FSlateBrush ChatBackgroundBrush;
 	FFriendsAndChatStyle& SetChatBackgroundBrush(const FSlateBrush& InChatBackgroundBrush);
@@ -282,3 +286,29 @@ struct FRIENDSANDCHAT_API FFriendsAndChatStyle
 	FSlateBrush ChatOptionsBackgroundBrush;
 	FFriendsAndChatStyle& SetChatOptionsBackgroundBrush(const FSlateBrush& InChatOptionsBackgroundBrush);
 };
+
+/** Manages the style which provides resources for the rich text widget. */
+class FRIENDSANDCHAT_API FFriendsAndChatModuleStyle
+{
+public:
+
+	static void Initialize(FFriendsAndChatStyle FriendStyle);
+
+	static void Shutdown();
+
+	/** reloads textures used by slate renderer */
+	static void ReloadTextures();
+
+	/** @return The Slate style set for the Shooter game */
+	static const ISlateStyle& Get();
+
+	static FName GetStyleSetName();
+
+private:
+
+	static TSharedRef< class FSlateStyleSet > Create(FFriendsAndChatStyle FriendStyle);
+
+private:
+
+	static TSharedPtr< class FSlateStyleSet > FriendsAndChatModuleStyleInstance;
+};
