#FriendsAndChat - Updates from code review

diff --git a/Engine/Source/Developer/FriendsAndChat/Private/FriendsAndChatManager.cpp b/Engine/Source/Developer/FriendsAndChat/Private/FriendsAndChatManager.cpp
--- a/Engine/Source/Developer/FriendsAndChat/Private/FriendsAndChatManager.cpp
+++ b/Engine/Source/Developer/FriendsAndChat/Private/FriendsAndChatManager.cpp
@@ -1367,7 +1367,7 @@ void FFriendsAndChatManager::AcceptGameInvite(const TSharedPtr<IFriendItem>& Fri
 
 	if(ApplicationViewModel.IsValid())
 	{
-		const FString AdditionalCommandline = TEXT("-invitesession=") + FriendItem->GetUniqueID()->ToString() + TEXT(" -invitefrom=") + FriendItem->GetGameSessionId();
+		const FString AdditionalCommandline = TEXT("-invitesession=") + FriendItem->GetGameSessionId() + TEXT(" -invitefrom=") + FriendItem->GetUniqueID()->ToString();
 		ApplicationViewModel->LaunchFriendApp(AdditionalCommandline);
 	}
 
diff --git a/Engine/Source/Developer/FriendsAndChat/Private/Models/FriendItem.cpp b/Engine/Source/Developer/FriendsAndChat/Private/Models/FriendItem.cpp
--- a/Engine/Source/Developer/FriendsAndChat/Private/Models/FriendItem.cpp
+++ b/Engine/Source/Developer/FriendsAndChat/Private/Models/FriendItem.cpp
@@ -87,7 +87,7 @@ const bool FFriendItem::IsOnline() const
 {
 	if(OnlineFriend.IsValid())
 	{
-		return OnlineFriend->GetPresence().Status.State == EOnlinePresenceState::Online ? true : false;
+		return OnlineFriend->GetPresence().Status.State != EOnlinePresenceState::Offline ? true : false;
 	}
 	return false;
 }
