- Changed the ReadFriendsList delegate to not be a multicast delegate, and instead have a delegate passed in to the ReadFriendsList function that's called on completion

diff --git a/Engine/Source/Developer/FriendsAndChat/Private/FriendsAndChatManager.cpp b/Engine/Source/Developer/FriendsAndChat/Private/FriendsAndChatManager.cpp
--- a/Engine/Source/Developer/FriendsAndChat/Private/FriendsAndChatManager.cpp
+++ b/Engine/Source/Developer/FriendsAndChat/Private/FriendsAndChatManager.cpp
@@ -62,7 +62,6 @@ void FFriendsAndChatManager::Login()
 			check( FriendsInterface.IsValid() )
 
 			// Create delegates for list refreshes
-			OnReadFriendsCompleteDelegate = FOnReadFriendsListCompleteDelegate::CreateSP( this, &FFriendsAndChatManager::OnReadFriendsListComplete );
 			OnQueryRecentPlayersCompleteDelegate = FOnQueryRecentPlayersCompleteDelegate::CreateRaw(this, &FFriendsAndChatManager::OnQueryRecentPlayersComplete);
 			OnFriendsListChangedDelegate = FOnFriendsChangeDelegate::CreateSP(this, &FFriendsAndChatManager::OnFriendsListChanged);
 			OnAcceptInviteCompleteDelegate = FOnAcceptInviteCompleteDelegate::CreateSP( this, &FFriendsAndChatManager::OnAcceptInviteComplete );
@@ -85,7 +84,6 @@ void FFriendsAndChatManager::Login()
 			FriendsInterface->AddOnFriendRemovedDelegate(OnFriendRemovedDelegate);
 			FriendsInterface->AddOnInviteRejectedDelegate(OnFriendInviteRejected);
 			FriendsInterface->AddOnInviteAcceptedDelegate(OnFriendInviteAccepted);
-			FriendsInterface->AddOnReadFriendsListCompleteDelegate( 0, OnReadFriendsCompleteDelegate );
 			FriendsInterface->AddOnAcceptInviteCompleteDelegate( 0, OnAcceptInviteCompleteDelegate );
 			FriendsInterface->AddOnDeleteFriendCompleteDelegate( 0, OnDeleteFriendCompleteDelegate );
 			FriendsInterface->AddOnSendInviteCompleteDelegate( 0, OnSendInviteCompleteDelegate );
@@ -137,7 +135,6 @@ void FFriendsAndChatManager::Logout()
 			OnlineSubMcp->GetFriendsInterface()->ClearOnFriendRemovedDelegate(OnFriendRemovedDelegate);
 			OnlineSubMcp->GetFriendsInterface()->ClearOnInviteRejectedDelegate(OnFriendInviteRejected);
 			OnlineSubMcp->GetFriendsInterface()->ClearOnInviteAcceptedDelegate(OnFriendInviteAccepted);
-			OnlineSubMcp->GetFriendsInterface()->ClearOnReadFriendsListCompleteDelegate(0, OnReadFriendsCompleteDelegate);
 			OnlineSubMcp->GetFriendsInterface()->ClearOnAcceptInviteCompleteDelegate(0, OnAcceptInviteCompleteDelegate);
 			OnlineSubMcp->GetFriendsInterface()->ClearOnDeleteFriendCompleteDelegate(0, OnDeleteFriendCompleteDelegate);
 			OnlineSubMcp->GetFriendsInterface()->ClearOnSendInviteCompleteDelegate(0, OnSendInviteCompleteDelegate);
@@ -729,7 +726,8 @@ void FFriendsAndChatManager::SetState( EFriendsAndManagerState::Type NewState )
 		break;
 	case EFriendsAndManagerState::RequestFriendsListRefresh:
 		{
-			if(FriendsInterface->ReadFriendsList( 0, EFriendsLists::ToString( EFriendsLists::Default ) ))
+			FOnReadFriendsListComplete Delegate = FOnReadFriendsListComplete::CreateSP(this, &FFriendsAndChatManager::OnReadFriendsListComplete);
+			if (FriendsInterface->ReadFriendsList(0, EFriendsLists::ToString(EFriendsLists::Default), Delegate))
 			{
 				SetState(EFriendsAndManagerState::RequestingFriendsList);
 			}
diff --git a/Engine/Source/Developer/FriendsAndChat/Private/FriendsAndChatManager.h b/Engine/Source/Developer/FriendsAndChat/Private/FriendsAndChatManager.h
--- a/Engine/Source/Developer/FriendsAndChat/Private/FriendsAndChatManager.h
+++ b/Engine/Source/Developer/FriendsAndChat/Private/FriendsAndChatManager.h
@@ -665,8 +665,6 @@ private:
 
 	// Holds the ticker delegate
 	FTickerDelegate UpdateFriendsTickerDelegate;
-	// Delegate when friend list has been read
-	FOnReadFriendsListCompleteDelegate OnReadFriendsCompleteDelegate;
 	// Delegate to use for querying for recent players 
 	FOnQueryRecentPlayersCompleteDelegate OnQueryRecentPlayersCompleteDelegate;
 	// Delegate when invites accepted
diff --git a/Engine/Source/Runtime/Online/IOS/OnlineSubsystemIOS/Private/OnlineFriendsInterfaceIOS.cpp b/Engine/Source/Runtime/Online/IOS/OnlineSubsystemIOS/Private/OnlineFriendsInterfaceIOS.cpp
--- a/Engine/Source/Runtime/Online/IOS/OnlineSubsystemIOS/Private/OnlineFriendsInterfaceIOS.cpp
+++ b/Engine/Source/Runtime/Online/IOS/OnlineSubsystemIOS/Private/OnlineFriendsInterfaceIOS.cpp
@@ -48,7 +48,7 @@ FOnlineFriendsIOS::FOnlineFriendsIOS(FOnlineSubsystemIOS* InSubsystem)
 	IdentityInterface = (FOnlineIdentityIOS*)InSubsystem->GetIdentityInterface().Get();
 }
 
-bool FOnlineFriendsIOS::ReadFriendsList(int32 LocalUserNum, const FString& ListName)
+bool FOnlineFriendsIOS::ReadFriendsList(int32 LocalUserNum, const FString& ListName, const FOnReadFriendsListComplete& Delegate /*= FOnReadFriendsListComplete()*/)
 {
 	UE_LOG(LogOnline, Verbose, TEXT("FOnlineFriendsIOS::ReadFriendsList()"));
 	bool bSuccessfullyBeganReadFriends = false;
@@ -70,7 +70,7 @@ bool FOnlineFriendsIOS::ReadFriendsList(int32 LocalUserNum, const FString& ListN
 						// Report back to the game thread whether this succeeded.
 						[FIOSAsyncTask CreateTaskWithBlock : ^ bool(void)
 						{
-							TriggerOnReadFriendsListCompleteDelegates(0, false, ListName, ErrorStr);
+							Delegate.ExecuteIfBound(0, false, ListName, ErrorStr);
 							return true;
 						}];
 				}
@@ -112,7 +112,7 @@ bool FOnlineFriendsIOS::ReadFriendsList(int32 LocalUserNum, const FString& ListN
 								// Report back to the game thread whether this succeeded.
 								[FIOSAsyncTask CreateTaskWithBlock : ^ bool(void)
 								{
-									TriggerOnReadFriendsListCompleteDelegates(0, bWasSuccessful, ListName, ErrorStr);
+									Delegate.ExecuteIfBound(0, bWasSuccessful, ListName, ErrorStr);
 									return true;
 								}];
 						}
@@ -125,7 +125,7 @@ bool FOnlineFriendsIOS::ReadFriendsList(int32 LocalUserNum, const FString& ListN
 	else
 	{
 		// no gamecenter login means we cannot read the friends.
-		TriggerOnReadFriendsListCompleteDelegates( 0 , false, ListName, TEXT("not logged in") );
+		Delegate.ExecuteIfBound( 0 , false, ListName, TEXT("not logged in") );
 	}
 
 	return bSuccessfullyBeganReadFriends;
diff --git a/Engine/Source/Runtime/Online/IOS/OnlineSubsystemIOS/Private/OnlineFriendsInterfaceIOS.h b/Engine/Source/Runtime/Online/IOS/OnlineSubsystemIOS/Private/OnlineFriendsInterfaceIOS.h
--- a/Engine/Source/Runtime/Online/IOS/OnlineSubsystemIOS/Private/OnlineFriendsInterfaceIOS.h
+++ b/Engine/Source/Runtime/Online/IOS/OnlineSubsystemIOS/Private/OnlineFriendsInterfaceIOS.h
@@ -88,7 +88,7 @@ public:
 
 	// IOnlineFriends
 
-	virtual bool ReadFriendsList(int32 LocalUserNum, const FString& ListName) override;
+	virtual bool ReadFriendsList(int32 LocalUserNum, const FString& ListName, const FOnReadFriendsListComplete& Delegate = FOnReadFriendsListComplete()) override;
 	virtual bool DeleteFriendsList(int32 LocalUserNum, const FString& ListName) override;
 	virtual bool SendInvite(int32 LocalUserNum, const FUniqueNetId& FriendId, const FString& ListName) override;
 	virtual bool AcceptInvite(int32 LocalUserNum, const FUniqueNetId& FriendId, const FString& ListName) override;
diff --git a/Engine/Source/Runtime/Online/OnlineSubsystem/Public/Interfaces/OnlineFriendsInterface.h b/Engine/Source/Runtime/Online/OnlineSubsystem/Public/Interfaces/OnlineFriendsInterface.h
--- a/Engine/Source/Runtime/Online/OnlineSubsystem/Public/Interfaces/OnlineFriendsInterface.h
+++ b/Engine/Source/Runtime/Online/OnlineSubsystem/Public/Interfaces/OnlineFriendsInterface.h
@@ -47,8 +47,7 @@ typedef FOnFriendsChange::FDelegate FOnFriendsChangeDelegate;
  * @param ListName name of the friends list that was operated on
  * @param ErrorStr string representing the error condition
  */
-DECLARE_MULTICAST_DELEGATE_FourParams(FOnReadFriendsListComplete, int32, bool, const FString&, const FString&);
-typedef FOnReadFriendsListComplete::FDelegate FOnReadFriendsListCompleteDelegate;
+DECLARE_DELEGATE_FourParams(FOnReadFriendsListComplete, int32, bool, const FString&, const FString&);
 
 /**
  * Delegate used when the friends list delete request has completed
@@ -212,17 +211,7 @@ public:
 	 *
 	 * @return true if the read request was started successfully, false otherwise
 	 */
-	virtual bool ReadFriendsList(int32 LocalUserNum, const FString& ListName) = 0;
-
-	/**
-	 * Delegate used when the friends list read request has completed
-	 *
-	 * @param LocalUserNum the controller number of the associated user that made the request
-	 * @param bWasSuccessful true if the async action completed without error, false if there was an error
-	 * @param ListName name of the friends list that was operated on
-	 * @param ErrorStr string representing the error condition
-	 */
-	DEFINE_ONLINE_PLAYER_DELEGATE_THREE_PARAM(MAX_LOCAL_PLAYERS, OnReadFriendsListComplete, bool, const FString&, const FString&);
+	virtual bool ReadFriendsList(int32 LocalUserNum, const FString& ListName, const FOnReadFriendsListComplete& Delegate = FOnReadFriendsListComplete()) = 0;
 
 	/**
 	 * Starts an async task that deletes the named friends list for the player 
diff --git a/Engine/Source/Runtime/Online/OnlineSubsystemFacebook/Private/IOS/OnlineFriendsFacebook.cpp b/Engine/Source/Runtime/Online/OnlineSubsystemFacebook/Private/IOS/OnlineFriendsFacebook.cpp
--- a/Engine/Source/Runtime/Online/OnlineSubsystemFacebook/Private/IOS/OnlineFriendsFacebook.cpp
+++ b/Engine/Source/Runtime/Online/OnlineSubsystemFacebook/Private/IOS/OnlineFriendsFacebook.cpp
@@ -69,7 +69,7 @@ FOnlineFriendsFacebook::~FOnlineFriendsFacebook()
 	SharingInterface = NULL;
 }
 
-bool FOnlineFriendsFacebook::ReadFriendsList(int32 LocalUserNum, const FString& ListName)
+bool FOnlineFriendsFacebook::ReadFriendsList(int32 LocalUserNum, const FString& ListName, const FOnReadFriendsListComplete& Delegate /*= FOnReadFriendsListComplete()*/)
 {
 	bool bRequestTriggered = false;
 	UE_LOG(LogOnline, Verbose, TEXT("FOnlineFriendsFacebook::ReadFriendsList()"));
@@ -82,13 +82,14 @@ bool FOnlineFriendsFacebook::ReadFriendsList(int32 LocalUserNum, const FString&
 		RequestFriendsReadPermissionsDelegate = FOnRequestNewReadPermissionsCompleteDelegate::CreateRaw(this, &FOnlineFriendsFacebook::OnReadFriendsPermissionsUpdated);
 		SharingInterface->AddOnRequestNewReadPermissionsCompleteDelegate(LocalUserNum, RequestFriendsReadPermissionsDelegate);
 		SharingInterface->RequestNewReadPermissions(LocalUserNum, EOnlineSharingReadCategory::Friends);
+		OnReadFriendsListCompleteDelegate = Delegate;
 	}
 	else
 	{
 		UE_LOG(LogOnline, Verbose, TEXT("Cannot read friends if we are not logged into facebook."));
 
 		// Cannot read friends if we are not logged into facebook.
-		TriggerOnReadFriendsListCompleteDelegates( LocalUserNum , false, ListName, TEXT("not logged in.") );
+		Delegate.ExecuteIfBound(LocalUserNum, false, ListName, TEXT("not logged in."));
 	}
 
 	return bRequestTriggered;
@@ -193,7 +194,7 @@ void FOnlineFriendsFacebook::OnReadFriendsPermissionsUpdated(int32 LocalUserNum,
 	else
 	{
 		// Permissions werent applied so we cannot read friends.
-		TriggerOnReadFriendsListCompleteDelegates(LocalUserNum, false, EFriendsLists::ToString(EFriendsLists::Default), TEXT("No read permissions"));
+		OnReadFriendsListCompleteDelegate.ExecuteIfBound(LocalUserNum, false, EFriendsLists::ToString(EFriendsLists::Default), TEXT("No read permissions"));
 	}
 }
 
@@ -268,7 +269,7 @@ void FOnlineFriendsFacebook::ReadFriendsUsingGraphPath(int32 LocalUserNum, const
 					}
 
 					// Did this operation complete? Let whoever is listening know.
-					TriggerOnReadFriendsListCompleteDelegates( LocalUserNum, bSuccess, ListName, FString() );
+					OnReadFriendsListCompleteDelegate.ExecuteIfBound(LocalUserNum, bSuccess, ListName, FString());
 				}
 			];
 		}
diff --git a/Engine/Source/Runtime/Online/OnlineSubsystemFacebook/Private/IOS/OnlineFriendsFacebook.h b/Engine/Source/Runtime/Online/OnlineSubsystemFacebook/Private/IOS/OnlineFriendsFacebook.h
--- a/Engine/Source/Runtime/Online/OnlineSubsystemFacebook/Private/IOS/OnlineFriendsFacebook.h
+++ b/Engine/Source/Runtime/Online/OnlineSubsystemFacebook/Private/IOS/OnlineFriendsFacebook.h
@@ -83,7 +83,7 @@ public:
 
 	// IOnlineFriends
 
-	virtual bool ReadFriendsList(int32 LocalUserNum, const FString& ListName) override;
+	virtual bool ReadFriendsList(int32 LocalUserNum, const FString& ListName, const FOnReadFriendsListComplete& Delegate = FOnReadFriendsListComplete()) override;
 	virtual bool DeleteFriendsList(int32 LocalUserNum, const FString& ListName) override;
 	virtual bool SendInvite(int32 LocalUserNum, const FUniqueNetId& FriendId, const FString& ListName) override;
 	virtual bool AcceptInvite(int32 LocalUserNum, const FUniqueNetId& FriendId, const FString& ListName) override;
@@ -144,6 +144,8 @@ private:
 	/** Delegate used to notify this interface that permissions are updated, and that we can now read the friends list */
 	FOnRequestNewReadPermissionsCompleteDelegate RequestFriendsReadPermissionsDelegate;
 
+	/** Delegate called when reading the friends list is completed */
+	FOnReadFriendsListComplete OnReadFriendsListCompleteDelegate;
 
 private:
 
diff --git a/Engine/Source/Runtime/Online/OnlineSubsystemFacebook/Private/Windows/OnlineFriendsFacebook.cpp b/Engine/Source/Runtime/Online/OnlineSubsystemFacebook/Private/Windows/OnlineFriendsFacebook.cpp
--- a/Engine/Source/Runtime/Online/OnlineSubsystemFacebook/Private/Windows/OnlineFriendsFacebook.cpp
+++ b/Engine/Source/Runtime/Online/OnlineSubsystemFacebook/Private/Windows/OnlineFriendsFacebook.cpp
@@ -61,7 +61,7 @@ FOnlineFriendsFacebook::~FOnlineFriendsFacebook()
 {
 }
 
-bool FOnlineFriendsFacebook::ReadFriendsList(int32 LocalUserNum, const FString& ListName)
+bool FOnlineFriendsFacebook::ReadFriendsList(int32 LocalUserNum, const FString& ListName, const FOnReadFriendsListComplete& Delegate /*= FOnReadFriendsListComplete()*/)
 {
 	FString AccessToken;
 	FString ErrorStr;
@@ -96,7 +96,7 @@ bool FOnlineFriendsFacebook::ReadFriendsList(int32 LocalUserNum, const FString&
 	if (!ErrorStr.IsEmpty())
 	{
 		UE_LOG_ONLINE(Warning, TEXT("ReadFriendsList request failed. %s"), *ErrorStr);
-		TriggerOnReadFriendsListCompleteDelegates(LocalUserNum, false, ListName, ErrorStr);
+		Delegate.ExecuteIfBound(LocalUserNum, false, ListName, ErrorStr);
 		return false;
 	}
 	
@@ -119,7 +119,7 @@ bool FOnlineFriendsFacebook::ReadFriendsList(int32 LocalUserNum, const FString&
 	FriendsQueryUrl = FriendsQueryUrl.Replace(TEXT("`token"), *AccessToken, ESearchCase::IgnoreCase);
 	
 	// kick off http request to read friends
-	HttpRequest->OnProcessRequestComplete().BindRaw(this, &FOnlineFriendsFacebook::QueryFriendsList_HttpRequestComplete);
+	HttpRequest->OnProcessRequestComplete().BindRaw(this, &FOnlineFriendsFacebook::QueryFriendsList_HttpRequestComplete, Delegate);
 	HttpRequest->SetURL(FriendsQueryUrl);
 	HttpRequest->SetHeader(TEXT("Content-Type"), TEXT("application/json"));
 	HttpRequest->SetVerb(TEXT("GET"));
@@ -224,7 +224,7 @@ bool FOnlineFriendsFacebook::GetRecentPlayers(const FUniqueNetId& UserId, TArray
 	return false;
 }
 
-void FOnlineFriendsFacebook::QueryFriendsList_HttpRequestComplete(FHttpRequestPtr HttpRequest, FHttpResponsePtr HttpResponse, bool bSucceeded)
+void FOnlineFriendsFacebook::QueryFriendsList_HttpRequestComplete(FHttpRequestPtr HttpRequest, FHttpResponsePtr HttpResponse, bool bSucceeded, FOnReadFriendsListComplete Delegate)
 {
 	bool bResult = false;
 	FString ResponseStr, ErrorStr;
@@ -301,5 +301,5 @@ void FOnlineFriendsFacebook::QueryFriendsList_HttpRequestComplete(FHttpRequestPt
 		UE_LOG(LogOnline, Warning, TEXT("Query friends list request failed. %s"), *ErrorStr);
 	}
 
-	TriggerOnReadFriendsListCompleteDelegates(PendingFriendsQuery.LocalUserNum, bResult, EFriendsLists::ToString(EFriendsLists::Default), ErrorStr);
+	Delegate.ExecuteIfBound(PendingFriendsQuery.LocalUserNum, bResult, EFriendsLists::ToString(EFriendsLists::Default), ErrorStr);
 }
\ No newline at end of file
diff --git a/Engine/Source/Runtime/Online/OnlineSubsystemFacebook/Private/Windows/OnlineFriendsFacebook.h b/Engine/Source/Runtime/Online/OnlineSubsystemFacebook/Private/Windows/OnlineFriendsFacebook.h
--- a/Engine/Source/Runtime/Online/OnlineSubsystemFacebook/Private/Windows/OnlineFriendsFacebook.h
+++ b/Engine/Source/Runtime/Online/OnlineSubsystemFacebook/Private/Windows/OnlineFriendsFacebook.h
@@ -81,7 +81,7 @@ public:
 
 	// IOnlineFriends
 
-	virtual bool ReadFriendsList(int32 LocalUserNum, const FString& ListName) override;
+	virtual bool ReadFriendsList(int32 LocalUserNum, const FString& ListName, const FOnReadFriendsListComplete& Delegate = FOnReadFriendsListComplete()) override;
 	virtual bool DeleteFriendsList(int32 LocalUserNum, const FString& ListName) override;
 	virtual bool SendInvite(int32 LocalUserNum, const FUniqueNetId& FriendId, const FString& ListName) override;
 	virtual bool AcceptInvite(int32 LocalUserNum, const FUniqueNetId& FriendId, const FString& ListName) override;
@@ -117,7 +117,7 @@ private:
 	/**
 	 * Delegate called when a user /me request from facebook is complete
 	 */
-	void QueryFriendsList_HttpRequestComplete(FHttpRequestPtr HttpRequest, FHttpResponsePtr HttpResponse, bool bSucceeded);
+	void QueryFriendsList_HttpRequestComplete(FHttpRequestPtr HttpRequest, FHttpResponsePtr HttpResponse, bool bSucceeded, FOnReadFriendsListComplete Delegate);
 
 	/** For accessing identity/token info of user logged in */
 	FOnlineSubsystemFacebook* FacebookSubsystem;
diff --git a/Engine/Source/Runtime/Online/OnlineSubsystemSteam/Private/OnlineFriendsInterfaceSteam.cpp b/Engine/Source/Runtime/Online/OnlineSubsystemSteam/Private/OnlineFriendsInterfaceSteam.cpp
--- a/Engine/Source/Runtime/Online/OnlineSubsystemSteam/Private/OnlineFriendsInterfaceSteam.cpp
+++ b/Engine/Source/Runtime/Online/OnlineSubsystemSteam/Private/OnlineFriendsInterfaceSteam.cpp
@@ -59,7 +59,7 @@ FOnlineFriendsSteam::FOnlineFriendsSteam(FOnlineSubsystemSteam* InSteamSubsystem
 	SteamFriendsPtr = SteamFriends();
 }
 
-bool FOnlineFriendsSteam::ReadFriendsList(int32 LocalUserNum, const FString& ListName)
+bool FOnlineFriendsSteam::ReadFriendsList(int32 LocalUserNum, const FString& ListName, const FOnReadFriendsListComplete& Delegate /*= FOnReadFriendsListComplete()*/)
 {
 	FString ErrorStr;
 	if (!ListName.Equals(DefaultFriendsList, ESearchCase::IgnoreCase))
@@ -71,7 +71,7 @@ bool FOnlineFriendsSteam::ReadFriendsList(int32 LocalUserNum, const FString& Lis
 		SteamUserPtr->BLoggedOn() &&
 		SteamFriendsPtr != NULL)
 	{
-		SteamSubsystem->QueueAsyncTask(new FOnlineAsyncTaskSteamReadFriendsList(this,LocalUserNum));
+		SteamSubsystem->QueueAsyncTask(new FOnlineAsyncTaskSteamReadFriendsList(this, LocalUserNum, Delegate));
 	}
 	else
 	{
@@ -79,7 +79,7 @@ bool FOnlineFriendsSteam::ReadFriendsList(int32 LocalUserNum, const FString& Lis
 	}
 	if (!ErrorStr.IsEmpty())
 	{
-		TriggerOnReadFriendsListCompleteDelegates(LocalUserNum, false, ListName, ErrorStr);
+		Delegate.ExecuteIfBound(LocalUserNum, false, ListName, ErrorStr);
 		return false;
 	}
 	return true;
@@ -252,5 +252,5 @@ void FOnlineAsyncTaskSteamReadFriendsList::TriggerDelegates(void)
 {
 	FOnlineAsyncTask::TriggerDelegates();
 
-	FriendsPtr->TriggerOnReadFriendsListCompleteDelegates(LocalUserNum,true,DefaultFriendsList,FString());
+	Delegate.ExecuteIfBound(LocalUserNum, true, DefaultFriendsList, FString());
 }
diff --git a/Engine/Source/Runtime/Online/OnlineSubsystemSteam/Private/OnlineFriendsInterfaceSteam.h b/Engine/Source/Runtime/Online/OnlineSubsystemSteam/Private/OnlineFriendsInterfaceSteam.h
--- a/Engine/Source/Runtime/Online/OnlineSubsystemSteam/Private/OnlineFriendsInterfaceSteam.h
+++ b/Engine/Source/Runtime/Online/OnlineSubsystemSteam/Private/OnlineFriendsInterfaceSteam.h
@@ -103,7 +103,7 @@ public:
 
 	// IOnlineFriends
 
-	virtual bool ReadFriendsList(int32 LocalUserNum, const FString& ListName) override;
+	virtual bool ReadFriendsList(int32 LocalUserNum, const FString& ListName, const FOnReadFriendsListComplete& Delegate = FOnReadFriendsListComplete()) override;
 	virtual bool DeleteFriendsList(int32 LocalUserNum, const FString& ListName) override;
 	virtual bool SendInvite(int32 LocalUserNum, const FUniqueNetId& FriendId, const FString& ListName) override;
 	virtual bool AcceptInvite(int32 LocalUserNum, const FUniqueNetId& FriendId, const FString& ListName) override;
@@ -129,16 +129,20 @@ class FOnlineAsyncTaskSteamReadFriendsList :
 	/** The user that is triggering the event */
 	int32 LocalUserNum;
 
+	FOnReadFriendsListComplete Delegate;
+
 public:
 	/**
 	 * Inits the pointer used to trigger the delegates on
 	 *
 	 * @param InFriendsPtr the interface to call the delegates on
 	 * @param InLocalUserNum the local user that requested the read
+	 * @param InDelegate the delegate that will be called when reading the friends list is complete
 	 */
-	FOnlineAsyncTaskSteamReadFriendsList(FOnlineFriendsSteam* InFriendsPtr, int32 InLocalUserNum) :
+	FOnlineAsyncTaskSteamReadFriendsList(FOnlineFriendsSteam* InFriendsPtr, int32 InLocalUserNum, const FOnReadFriendsListComplete& InDelegate) :
 		FriendsPtr(InFriendsPtr),
-		LocalUserNum(InLocalUserNum)
+		LocalUserNum(InLocalUserNum),
+		Delegate(InDelegate)
 	{
 		check(FriendsPtr);
 	}
diff --git a/Engine/Source/Runtime/Online/OnlineSubsystemUtils/Private/Tests/TestFriendsInterface.cpp b/Engine/Source/Runtime/Online/OnlineSubsystemUtils/Private/Tests/TestFriendsInterface.cpp
--- a/Engine/Source/Runtime/Online/OnlineSubsystemUtils/Private/Tests/TestFriendsInterface.cpp
+++ b/Engine/Source/Runtime/Online/OnlineSubsystemUtils/Private/Tests/TestFriendsInterface.cpp
@@ -19,7 +19,6 @@ void FTestFriendsInterface::Test(UWorld* InWorld, const TArray<FString>& Invites
 		}
 
 		// Add our delegate for the async call
-		OnReadFriendsCompleteDelegate = FOnReadFriendsListCompleteDelegate::CreateRaw(this, &FTestFriendsInterface::OnReadFriendsComplete);
 		OnAcceptInviteCompleteDelegate = FOnAcceptInviteCompleteDelegate::CreateRaw(this, &FTestFriendsInterface::OnAcceptInviteComplete);
 		OnSendInviteCompleteDelegate = FOnSendInviteCompleteDelegate::CreateRaw(this, &FTestFriendsInterface::OnSendInviteComplete);
 		OnDeleteFriendCompleteDelegate = FOnDeleteFriendCompleteDelegate::CreateRaw(this, &FTestFriendsInterface::OnDeleteFriendComplete);
@@ -27,7 +26,6 @@ void FTestFriendsInterface::Test(UWorld* InWorld, const TArray<FString>& Invites
 		OnQueryRecentPlayersCompleteDelegate = FOnQueryRecentPlayersCompleteDelegate::CreateRaw(this, &FTestFriendsInterface::OnQueryRecentPlayersComplete);
 		
 
-		OnlineSub->GetFriendsInterface()->AddOnReadFriendsListCompleteDelegate(0, OnReadFriendsCompleteDelegate);
 		OnlineSub->GetFriendsInterface()->AddOnAcceptInviteCompleteDelegate(0, OnAcceptInviteCompleteDelegate);
 		OnlineSub->GetFriendsInterface()->AddOnSendInviteCompleteDelegate(0, OnSendInviteCompleteDelegate);
 		OnlineSub->GetFriendsInterface()->AddOnDeleteFriendCompleteDelegate(0, OnDeleteFriendCompleteDelegate);
@@ -60,7 +58,8 @@ void FTestFriendsInterface::StartNextTest()
 {
 	if (bReadFriendsList)
 	{
-		OnlineSub->GetFriendsInterface()->ReadFriendsList(0, FriendsListName);
+		FOnReadFriendsListComplete Delegate = FOnReadFriendsListComplete::CreateRaw(this, &FTestFriendsInterface::OnReadFriendsComplete);
+		OnlineSub->GetFriendsInterface()->ReadFriendsList(0, FriendsListName, Delegate);
 	}
 	else if (bQueryRecentPlayers)
 	{
@@ -100,7 +99,6 @@ void FTestFriendsInterface::FinishTest()
 		OnlineSub->GetFriendsInterface().IsValid())
 	{
 		// Clear delegates for the various async calls
-		OnlineSub->GetFriendsInterface()->ClearOnReadFriendsListCompleteDelegate(0, OnReadFriendsCompleteDelegate);
 		OnlineSub->GetFriendsInterface()->ClearOnAcceptInviteCompleteDelegate(0, OnAcceptInviteCompleteDelegate);
 		OnlineSub->GetFriendsInterface()->ClearOnSendInviteCompleteDelegate(0, OnSendInviteCompleteDelegate);
 		OnlineSub->GetFriendsInterface()->ClearOnDeleteFriendCompleteDelegate(0, OnDeleteFriendCompleteDelegate);
diff --git a/Engine/Source/Runtime/Online/OnlineSubsystemUtils/Private/Tests/TestFriendsInterface.h b/Engine/Source/Runtime/Online/OnlineSubsystemUtils/Private/Tests/TestFriendsInterface.h
--- a/Engine/Source/Runtime/Online/OnlineSubsystemUtils/Private/Tests/TestFriendsInterface.h
+++ b/Engine/Source/Runtime/Online/OnlineSubsystemUtils/Private/Tests/TestFriendsInterface.h
@@ -14,8 +14,6 @@
 	const FString SubsystemName;
 	/** The online interface to use for testing */
 	IOnlineSubsystem* OnlineSub;
-	/** Delegate to use for reading the friends list */
-	FOnReadFriendsListCompleteDelegate OnReadFriendsCompleteDelegate;
 	/** Delegate when invites accepted */
 	FOnAcceptInviteCompleteDelegate OnAcceptInviteCompleteDelegate;
 	/** Delegate when invites have been sent */
diff --git a/Engine/Source/Runtime/Online/OnlineSubsystemUtils/Private/Tests/TestSessionInterface.cpp b/Engine/Source/Runtime/Online/OnlineSubsystemUtils/Private/Tests/TestSessionInterface.cpp
--- a/Engine/Source/Runtime/Online/OnlineSubsystemUtils/Private/Tests/TestSessionInterface.cpp
+++ b/Engine/Source/Runtime/Online/OnlineSubsystemUtils/Private/Tests/TestSessionInterface.cpp
@@ -110,8 +110,6 @@ void FTestSessionInterface::Test(UWorld* InWorld, bool bTestLAN, bool bIsPresenc
 	Friends = OnlineSub->GetFriendsInterface();
 
 	// Define delegates
-	OnReadFriendsListCompleteDelegate = FOnReadFriendsListCompleteDelegate::CreateRaw(this, &FTestSessionInterface::OnReadFriendsListComplete);
-
 	OnCreateSessionCompleteDelegate = FOnCreateSessionCompleteDelegate::CreateRaw(this, &FTestSessionInterface::OnCreateSessionComplete);
 	OnStartSessionCompleteDelegate = FOnStartSessionCompleteDelegate::CreateRaw(this, &FTestSessionInterface::OnStartSessionComplete);
 	OnEndSessionCompleteDelegate = FOnEndSessionCompleteDelegate::CreateRaw(this, &FTestSessionInterface::OnEndSessionComplete);
@@ -134,8 +132,7 @@ void FTestSessionInterface::Test(UWorld* InWorld, bool bTestLAN, bool bIsPresenc
 	// Read friends list and cache it
 	if (Friends.IsValid())
 	{
-		Friends->AddOnReadFriendsListCompleteDelegate(0, OnReadFriendsListCompleteDelegate);
-		Friends->ReadFriendsList(0, EFriendsLists::ToString(EFriendsLists::Default));
+		Friends->ReadFriendsList(0, EFriendsLists::ToString(EFriendsLists::Default), FOnReadFriendsListComplete::CreateRaw(this, &FTestSessionInterface::OnReadFriendsListComplete));
 	}
 
 	// Setup sessions
@@ -177,7 +174,6 @@ void FTestSessionInterface::WorldDestroyed( UWorld* InWorld )
 void FTestSessionInterface::OnReadFriendsListComplete(int32 LocalUserNum, bool bWasSuccessful, const FString& ListName, const FString& ErrorStr)
 {
 	UE_LOG(LogOnline, Verbose, TEXT("OnReadFriendsListComplete LocalUserNum: %d bSuccess: %d %s"), LocalUserNum, bWasSuccessful, *ErrorStr);
-	Friends->ClearOnReadFriendsListCompleteDelegate(0, OnReadFriendsListCompleteDelegate);
 	if (bWasSuccessful)
 	{
 		Friends->GetFriendsList(LocalUserNum, ListName, FriendsCache);
diff --git a/Engine/Source/Runtime/Online/OnlineSubsystemUtils/Private/Tests/TestSessionInterface.h b/Engine/Source/Runtime/Online/OnlineSubsystemUtils/Private/Tests/TestSessionInterface.h
--- a/Engine/Source/Runtime/Online/OnlineSubsystemUtils/Private/Tests/TestSessionInterface.h
+++ b/Engine/Source/Runtime/Online/OnlineSubsystemUtils/Private/Tests/TestSessionInterface.h
@@ -36,8 +36,6 @@
 	/** Convenient access to friends interface */
 	IOnlineFriendsPtr Friends;
 
-	/** Delegate when a friends list read is complete */
-	FOnReadFriendsListCompleteDelegate OnReadFriendsListCompleteDelegate;
 	/** Delegate for handling an accepted invite */
 	FOnSessionInviteAcceptedDelegate OnSessionInviteAcceptedDelegate;
 
