Fix up message caps for different message types. Fixed a party chat message display issue. Changed the list view to use chat item view models provided from the manager. Fixed an issue with non-friends selected action changes when a friend status changes

diff --git a/Engine/Source/Developer/FriendsAndChat/Private/FriendsMessageManager.cpp b/Engine/Source/Developer/FriendsAndChat/Private/FriendsMessageManager.cpp
--- a/Engine/Source/Developer/FriendsAndChat/Private/FriendsMessageManager.cpp
+++ b/Engine/Source/Developer/FriendsAndChat/Private/FriendsMessageManager.cpp
@@ -2,13 +2,17 @@
 
 #include "FriendsAndChatPrivatePCH.h"
 #include "OnlineChatInterface.h"
+#include "ChatItemViewModel.h"
 
 #define LOCTEXT_NAMESPACE "FriendsMessageManager"
 // Message expiry time for different message types
-static const int32 GlobalMessageLifetime = 5 * 60.0f;  // 5 min
-static const int32 PartyMessageLifetime = 5 * 60.0f;  // 5 min
-static const int32 WhisperMessageLifetime = 5 * 60.0f;  // 5 min
+static const int32 GlobalMessageLifetime = 5 * 60;  // 5 min
+static const int32 PartyMessageLifetime = 5 * 60;  // 5 min
+static const int32 WhisperMessageLifetime = 5 * 60;  // 5 min
 static const int32 MessageStore = 200;
+static const int32 GlobalMaxStore = 100;
+static const int32 WhisperMaxStore = 100;
+static const int32 PartyMaxStore = 100;
 
 class FFriendsMessageManagerImpl
 	: public FFriendsMessageManager
@@ -49,7 +53,7 @@ public:
 		UnInitialize();
 	}
 
-	virtual const TArray<TSharedPtr<FFriendChatMessage> >& GetMessageList() const override
+	virtual const TArray<TSharedRef<FChatItemViewModel> >& GetMessageList() const override
 	{
 		return ReceivedMessages;
 	}
@@ -92,7 +96,7 @@ public:
 		return false;
 	}
 
-	virtual bool SendPrivateMessage(const FUniqueNetId& RecipientId, TSharedPtr< FFriendChatMessage > ChatMessage) override
+	virtual bool SendPrivateMessage(TSharedPtr<FUniqueNetId> UserID, const FText UserName, const FText MessageText) override
 	{
 		if (OnlineSub != nullptr &&
 			LoggedInUser.IsValid())
@@ -100,8 +104,16 @@ public:
 			IOnlineChatPtr ChatInterface = OnlineSub->GetChatInterface();
 			if(ChatInterface.IsValid())
 			{
-				AddMessage(ChatMessage.ToSharedRef());
-				return ChatInterface->SendPrivateChat(*LoggedInUser, RecipientId, ChatMessage->Message.ToString());
+				TSharedPtr< FFriendChatMessage > ChatItem = MakeShareable(new FFriendChatMessage());
+				ChatItem->FromName = UserName;
+				ChatItem->Message = MessageText;
+				ChatItem->MessageType = EChatMessageType::Whisper;
+				ChatItem->MessageTimeText = FText::AsTime(FDateTime::UtcNow());
+				ChatItem->ExpireTime = FDateTime::UtcNow() + FTimespan::FromSeconds(WhisperMessageLifetime);
+				ChatItem->bIsFromSelf = true;
+				ChatItem->SenderId = UserID;
+				AddMessage(ChatItem.ToSharedRef());
+				return ChatInterface->SendPrivateChat(*LoggedInUser, *UserID.Get(), MessageText.ToString());
 			}
 		}
 		return false;
@@ -114,8 +126,9 @@ public:
 		ChatItem->Message = FText::FromString(MsgBody);
 		ChatItem->MessageType = EChatMessageType::Party;
 		ChatItem->MessageTimeText = FText::AsTime(FDateTime::UtcNow());
-		ChatItem->ExpireTime = FDateTime::UtcNow() + PartyMessageLifetime;
+		ChatItem->ExpireTime = FDateTime::UtcNow() + FTimespan::FromSeconds(PartyMessageLifetime);
 		ChatItem->bIsFromSelf = false;
+		PartyMessagesCount++;
 		AddMessage(ChatItem.ToSharedRef());
 	}
 
@@ -166,6 +179,11 @@ private:
 		// Clear existing data
 		LogOut();
 
+		GlobalMessagesCount = 0;
+		WhisperMessagesCount = 0;
+		PartyMessagesCount = 0;
+		ReceivedMessages.Empty();
+
 		OnlineSub = IOnlineSubsystem::Get( TEXT( "MCP" ) );
 
 		if (OnlineSub != nullptr)
@@ -249,7 +267,8 @@ private:
 			ChatItem->MessageTimeText = FText::AsTime(FDateTime::UtcNow());
 			ChatItem->ExpireTime = FDateTime::UtcNow() + GlobalMessageLifetime;
 			ChatItem->bIsFromSelf = false;
-			AddMessage(ChatItem);
+			GlobalMessagesCount++;
+			AddMessage(ChatItem.ToSharedRef());
 		}
 	}
 
@@ -270,7 +289,8 @@ private:
 			ChatItem->MessageTimeText = FText::AsTime(FDateTime::UtcNow());
 			ChatItem->ExpireTime = FDateTime::UtcNow() + GlobalMessageLifetime;
 			ChatItem->bIsFromSelf = false;
-			AddMessage(ChatItem);
+			GlobalMessagesCount++;
+			AddMessage(ChatItem.ToSharedRef());
 		}
 	}
 
@@ -286,7 +306,7 @@ private:
 		ChatItem->Message = FText::FromString(*ChatMessage->GetBody());
 		ChatItem->MessageType = EChatMessageType::Global;
 		ChatItem->MessageTimeText = FText::AsTime(ChatMessage->GetTimestamp());
-		ChatItem->ExpireTime = ChatMessage->GetTimestamp() + GlobalMessageLifetime;
+		ChatItem->ExpireTime = ChatMessage->GetTimestamp() + FTimespan::FromSeconds(GlobalMessageLifetime);
 		ChatItem->bIsFromSelf = ChatMessage->GetUserId() == *LoggedInUser;
 		TSharedPtr<IFriendItem> FoundFriend = FFriendsAndChatManager::Get()->FindUser(ChatMessage->GetUserId());
 		if(FoundFriend.IsValid())
@@ -295,8 +315,8 @@ private:
 		}
 
 		ChatItem->MessageRef = ChatMessage;
-		AddMessage(ChatItem);
-
+		GlobalMessagesCount++;
+		AddMessage(ChatItem.ToSharedRef());
 	}
 
 	void OnChatPrivateMessageReceived(const FUniqueNetId& UserId, const TSharedRef<FChatMessage>& ChatMessage)
@@ -312,10 +332,11 @@ private:
 			ChatItem->Message = FText::FromString(*ChatMessage->GetBody());
 			ChatItem->MessageType = EChatMessageType::Whisper;
 			ChatItem->MessageTimeText = FText::AsTime(ChatMessage->GetTimestamp());
-			ChatItem->ExpireTime = ChatMessage->GetTimestamp() + WhisperMessageLifetime;
+			ChatItem->ExpireTime = ChatMessage->GetTimestamp() + FTimespan::FromSeconds(WhisperMessageLifetime);
 			ChatItem->bIsFromSelf = false;
 			ChatItem->MessageRef = ChatMessage;
-			AddMessage(ChatItem);
+			WhisperMessagesCount++;
+			AddMessage(ChatItem.ToSharedRef());
 		}
 	}
 
@@ -334,9 +355,10 @@ private:
 		}
 	}
 
-	void AddMessage(TSharedPtr< FFriendChatMessage > ChatItem)
+	void AddMessage(TSharedRef< FFriendChatMessage > ChatItem)
 	{
-		if(ReceivedMessages.Add(ChatItem) > MessageStore)
+		TSharedRef<FChatItemViewModel> NewMessage = FChatItemViewModelFactory::Create(ChatItem);
+		if(ReceivedMessages.Add(NewMessage) > MessageStore)
 		{
 			bool bGlobalTimeFound = false;
 			bool bPartyTimeFound = false;
@@ -344,35 +366,76 @@ private:
 			FDateTime CurrentTime = FDateTime::UtcNow();
 			for(int32 Index = 0; Index < ReceivedMessages.Num(); Index++)
 			{
-				TSharedPtr<FFriendChatMessage> Message = ReceivedMessages[Index];
-				if(Message->ExpireTime < CurrentTime)
+				TSharedRef<FChatItemViewModel> Message = ReceivedMessages[Index];
+				if(Message->GetExpireTime() < CurrentTime)
 				{
-					ReceivedMessages.Remove(Message);
+					RemoveMessage(Message);
 					Index--;
 				}
 				else
 				{
-					switch(Message->MessageType)
+					switch(Message->GetMessageType())
 					{
-						case EChatMessageType::Global : bGlobalTimeFound = true; break;
-						case EChatMessageType::Party : bPartyTimeFound = true; break;
-						case EChatMessageType::Whisper : bWhisperFound = true; break;
+						case EChatMessageType::Global :
+						{
+							if(GlobalMessagesCount > GlobalMaxStore)
+							{
+								RemoveMessage(Message);
+								Index--;
+							}
+							else
+							{
+								bGlobalTimeFound = true;
+							}
+						}
+						break;
+						case EChatMessageType::Party :
+						{
+							if(PartyMessagesCount > PartyMaxStore)
+							{
+								RemoveMessage(Message);
+								Index--;
+							}
+							else
+							{
+								bPartyTimeFound = true;
+							}
+						}
+						break;
+						case EChatMessageType::Whisper :
+						{
+							if(WhisperMessagesCount > WhisperMaxStore)
+							{
+								RemoveMessage(Message);
+								Index--;
+							}
+							else
+							{
+								bWhisperFound = true;
+							}
+						}
+						break;
 					}
 				}
-
 				if(ReceivedMessages.Num() < MessageStore || (bPartyTimeFound && bGlobalTimeFound && bWhisperFound))
 				{
 					break;
 				}
 			}
-			OnChatMessageRecieved().Broadcast(ChatItem.ToSharedRef());
 		}
-		else
+		OnChatMessageRecieved().Broadcast(NewMessage);
+	}
+
+	void RemoveMessage(TSharedRef<FChatItemViewModel> Message)
+	{
+		switch(Message->GetMessageType())
 		{
-			OnChatMessageRecieved().Broadcast(ChatItem.ToSharedRef());
+			case EChatMessageType::Global : GlobalMessagesCount--; break;
+			case EChatMessageType::Party : PartyMessagesCount--; break;
+			case EChatMessageType::Whisper : WhisperMessagesCount--; break;
 		}
+		ReceivedMessages.Remove(Message);
 	}
-
 private:
 
 	// Incoming delegates
@@ -394,11 +457,11 @@ private:
 	TSharedPtr<FUniqueNetId> LoggedInUser;
 	TArray<FString> RoomJoins;
 
-	TArray<TSharedPtr< FFriendChatMessage > > ReceivedMessages;
+	TArray<TSharedRef<FChatItemViewModel> > ReceivedMessages;
 
-	TArray<TSharedPtr< FFriendChatMessage > > GlobalMessages;
-	TArray<TSharedPtr< FFriendChatMessage > > WhisperMessages;
-	TArray<TSharedPtr< FFriendChatMessage > > PartyMessages;
+	int32 GlobalMessagesCount;
+	int32 WhisperMessagesCount;
+	int32 PartyMessagesCount;
 
 	bool bEnableEnterExitMessages;
 
diff --git a/Engine/Source/Developer/FriendsAndChat/Private/FriendsMessageManager.h b/Engine/Source/Developer/FriendsAndChat/Private/FriendsMessageManager.h
--- a/Engine/Source/Developer/FriendsAndChat/Private/FriendsMessageManager.h
+++ b/Engine/Source/Developer/FriendsAndChat/Private/FriendsMessageManager.h
@@ -57,13 +57,13 @@ public:
 
 	virtual void LogIn() = 0;
 	virtual void LogOut() = 0;
-	virtual const TArray<TSharedPtr<FFriendChatMessage> >& GetMessageList() const = 0;
+	virtual const TArray<TSharedRef<class FChatItemViewModel> >& GetMessageList() const = 0;
 	virtual void JoinPublicRoom(const FString& RoomName) = 0;
 	virtual bool SendRoomMessage(const FString& RoomName, const FString& MsgBody) = 0;
-	virtual bool SendPrivateMessage(const FUniqueNetId& RecipientId, TSharedPtr< FFriendChatMessage > ChatMessage) = 0;
+	virtual bool SendPrivateMessage(TSharedPtr<FUniqueNetId> UserID, const FText UserName, const FText MessageText) = 0;
 	virtual void InsertNetworkMessage(const FString& MsgBody) = 0;
 
-	DECLARE_EVENT_OneParam(FFriendsMessageManager, FOnChatMessageReceivedEvent, const TSharedRef<FFriendChatMessage> /*The chat message*/)
+	DECLARE_EVENT_OneParam(FFriendsMessageManager, FOnChatMessageReceivedEvent, const TSharedRef<class FChatItemViewModel> /*The chat message*/)
 	virtual FOnChatMessageReceivedEvent& OnChatMessageRecieved() = 0;
 
 	DECLARE_EVENT_OneParam(FFriendsMessageManager, FOnChatPublicRoomJoinedEvent, const FString& /*RoomName*/)
diff --git a/Engine/Source/Developer/FriendsAndChat/Private/Models/ChatItemViewModel.cpp b/Engine/Source/Developer/FriendsAndChat/Private/Models/ChatItemViewModel.cpp
--- a/Engine/Source/Developer/FriendsAndChat/Private/Models/ChatItemViewModel.cpp
+++ b/Engine/Source/Developer/FriendsAndChat/Private/Models/ChatItemViewModel.cpp
@@ -37,9 +37,9 @@ public:
 		return ChatMessage->FromName;
 	}
 
-	virtual TSharedRef<FFriendChatMessage> GetMessageItem() const override
+	virtual const TSharedPtr<FUniqueNetId> GetSenderID() const override
 	{
-		return ChatMessage;
+		return ChatMessage->SenderId;
 	}
 
 	virtual FText GetMessageTime() override
@@ -47,58 +47,32 @@ public:
 		return ChatMessage->MessageTimeText;
 	}
 
-	const bool IsFromSelf() const override
+	virtual const FDateTime GetExpireTime() override
 	{
-		return ChatMessage->bIsFromSelf;
+		return ChatMessage->ExpireTime;
 	}
 
-	virtual const float GetFadeAmountColor() const override
-	{
-		return Owner->GetTimeTransparency();
-	}
-
-	virtual const bool UseOverrideColor() const override
-	{
-		return Owner->GetOverrideColorSet();
-	}
-
-	virtual const FSlateColor GetOverrideColor() const override
-	{
-		return Owner->GetFontOverrideColor();
-	}
-
-	virtual void FriendNameSelected() const override
+	const bool IsFromSelf() const override
 	{
-		if(!IsFromSelf() || GetMessageType() == EChatMessageType::Whisper)
-		{
-			Owner->SetChannelUserClicked(ChatMessage);
-		}
-		else
-		{
-			Owner->SetChatChannel(GetMessageType());
-		}
+		return ChatMessage->bIsFromSelf;
 	}
 
 private:
 
-	FChatItemViewModelImpl(TSharedRef<FFriendChatMessage> ChatMessage, TSharedPtr<FChatViewModel> Owner)
+	FChatItemViewModelImpl(TSharedRef<FFriendChatMessage> ChatMessage)
 	: ChatMessage(ChatMessage)
-	, Owner(Owner)
 	{
 	}
 
 private:
 	TSharedRef<FFriendChatMessage> ChatMessage;
-	TSharedPtr<FChatViewModel> Owner;
 
 private:
 	friend FChatItemViewModelFactory;
 };
 
-TSharedRef< FChatItemViewModel > FChatItemViewModelFactory::Create(
-	const TSharedRef<FFriendChatMessage>& ChatMessage
-	, TSharedRef<FChatViewModel> Owner)
+TSharedRef< FChatItemViewModel > FChatItemViewModelFactory::Create(const TSharedRef<FFriendChatMessage>& ChatMessage)
 {
-	TSharedRef< FChatItemViewModelImpl > ViewModel(new FChatItemViewModelImpl(ChatMessage, Owner));
+	TSharedRef< FChatItemViewModelImpl > ViewModel(new FChatItemViewModelImpl(ChatMessage));
 	return ViewModel;
-}
\ No newline at end of file
+}
diff --git a/Engine/Source/Developer/FriendsAndChat/Private/Models/ChatItemViewModel.h b/Engine/Source/Developer/FriendsAndChat/Private/Models/ChatItemViewModel.h
--- a/Engine/Source/Developer/FriendsAndChat/Private/Models/ChatItemViewModel.h
+++ b/Engine/Source/Developer/FriendsAndChat/Private/Models/ChatItemViewModel.h
@@ -16,13 +16,10 @@ public:
 	virtual const EChatMessageType::Type GetMessageType() const =0;
 	virtual const FText GetFriendNameDisplayText() const = 0;
 	virtual const FText GetFriendName() const = 0;
-	virtual TSharedRef<FFriendChatMessage> GetMessageItem() const = 0;
+	virtual const TSharedPtr<FUniqueNetId> GetSenderID() const = 0;
 	virtual FText GetMessageTime() = 0;
+	virtual const FDateTime GetExpireTime() = 0;
 	virtual const bool IsFromSelf() const = 0;
-	virtual const float GetFadeAmountColor() const = 0;
-	virtual const bool UseOverrideColor() const = 0;
-	virtual const FSlateColor GetOverrideColor() const = 0;
-	virtual void FriendNameSelected() const = 0;
 };
 
 /**
@@ -31,4 +28,4 @@ public:
  * @return the newly created ChaItemtViewModel implementation.
  */
 FACTORY(TSharedRef< FChatItemViewModel >, FChatItemViewModel,
-	const TSharedRef<FFriendChatMessage>& FFriendChatMessage, TSharedRef<class FChatViewModel> Owner);
\ No newline at end of file
+	const TSharedRef<FFriendChatMessage>& FFriendChatMessage);
\ No newline at end of file
diff --git a/Engine/Source/Developer/FriendsAndChat/Private/Models/ChatViewModel.cpp b/Engine/Source/Developer/FriendsAndChat/Private/Models/ChatViewModel.cpp
--- a/Engine/Source/Developer/FriendsAndChat/Private/Models/ChatViewModel.cpp
+++ b/Engine/Source/Developer/FriendsAndChat/Private/Models/ChatViewModel.cpp
@@ -5,9 +5,6 @@
 #include "ChatViewModel.h"
 #include "OnlineChatInterface.h"
 
-// Holds how many message to display
-static const int32 MessageDisplay = 200;
-
 class FChatViewModelImpl
 	: public FChatViewModel
 {
@@ -48,7 +45,7 @@ public:
 		return OverrideColor;
 	}
 
-	virtual TArray< TSharedRef<class FChatItemViewModel > >& GetFilteredChatList() override
+	virtual TArray< TSharedRef<FChatItemViewModel > >& GetFilteredChatList() override
 	{
 		return FilteredChatLists;
 	}
@@ -74,12 +71,12 @@ public:
 	{
 		if(ItemSelected->GetMessageType() == EChatMessageType::Whisper)
 		{
-			if(ItemSelected->GetMessageItem()->SenderId.IsValid())
+			if(ItemSelected->GetSenderID().IsValid())
 			{
-				TSharedPtr< IFriendItem > ExistingFriend = FFriendsAndChatManager::Get()->FindUser(*ItemSelected->GetMessageItem()->SenderId);
+				TSharedPtr< IFriendItem > ExistingFriend = FFriendsAndChatManager::Get()->FindUser(*ItemSelected->GetSenderID());
 				if(ExistingFriend.IsValid() && ExistingFriend->GetInviteStatus() == EInviteStatus::Accepted)
 				{
-					SetWhisperChannel(GetRecentFriend(ItemSelected->GetFriendName(), ItemSelected->GetMessageItem()->SenderId));
+					SetWhisperChannel(GetRecentFriend(ItemSelected->GetFriendName(), ItemSelected->GetSenderID()));
 				}
 			}
 		}
@@ -149,9 +146,9 @@ public:
 				{
 					FFriendsAndChatManager::Get()->SendGameInvite(*SelectedFriend->UserID.Get());
 				}
-				else if(SelectedFriend->SelectedMessage.IsValid() && SelectedFriend->SelectedMessage->MessageRef.IsValid())
+				else if(SelectedFriend->SelectedMessage.IsValid() && SelectedFriend->SelectedMessage->GetSenderID().IsValid())
 				{
-					FFriendsAndChatManager::Get()->SendGameInvite(SelectedFriend->SelectedMessage->MessageRef->GetUserId());
+					FFriendsAndChatManager::Get()->SendGameInvite(*SelectedFriend->SelectedMessage->GetSenderID().Get());
 				}
 			}
 		}
@@ -190,36 +187,39 @@ public:
 		bHasActionPending = false;
 	}
 
-	virtual void SetChannelUserClicked(const TSharedRef<FFriendChatMessage> ChatItemSelected) override
+	virtual void SetChannelUserClicked(const TSharedRef<FChatItemViewModel> ChatItemSelected) override
 	{
-		bool bFoundUser = false;
-		TSharedPtr< IFriendItem > ExistingFriend = NULL;
-
-		if(ChatItemSelected->SenderId.IsValid())
+		if(ChatItemSelected->IsFromSelf() && ChatItemSelected->GetMessageType() == EChatMessageType::Global)
 		{
-			ExistingFriend = FFriendsAndChatManager::Get()->FindUser(*ChatItemSelected->SenderId.Get());
+			SetChatChannel(EChatMessageType::Global);
 		}
-		else if(ChatItemSelected->MessageRef.IsValid())
+		else
 		{
-			ExistingFriend = FFriendsAndChatManager::Get()->FindUser(ChatItemSelected->MessageRef->GetUserId());
-		}
+			bool bFoundUser = false;
+			TSharedPtr< IFriendItem > ExistingFriend = NULL;
 
-		if(ExistingFriend.IsValid() && ExistingFriend->GetInviteStatus() == EInviteStatus::Accepted)
-		{
-			bFoundUser = true;
-			SetWhisperChannel(GetRecentFriend(FText::FromString(ExistingFriend->GetName()), ExistingFriend->GetUniqueID()));
-		}
+			if(ChatItemSelected->GetSenderID().IsValid())
+			{
+				ExistingFriend = FFriendsAndChatManager::Get()->FindUser(*ChatItemSelected->GetSenderID().Get());
+			}
 
-		if(!bFoundUser)
-		{
-			SetChatChannel(ChatItemSelected->MessageType);
-			SelectedFriend = MakeShareable(new FSelectedFriend());
-			SelectedFriend->FriendName = ChatItemSelected->FromName;
-			SelectedFriend->MessageType = ChatItemSelected->MessageType;
-			SelectedFriend->ViewModel = nullptr;
-			SelectedFriend->SelectedMessage = ChatItemSelected;
-			bHasActionPending = true;
-			bAllowJoinGame = FFriendsAndChatManager::Get()->IsInJoinableGameSession();
+			if(ExistingFriend.IsValid() && ExistingFriend->GetInviteStatus() == EInviteStatus::Accepted)
+			{
+				bFoundUser = true;
+				SetWhisperChannel(GetRecentFriend(FText::FromString(ExistingFriend->GetName()), ExistingFriend->GetUniqueID()));
+			}
+
+			if(!bFoundUser)
+			{
+				SetChatChannel(ChatItemSelected->GetMessageType());
+				SelectedFriend = MakeShareable(new FSelectedFriend());
+				SelectedFriend->FriendName = ChatItemSelected->GetFriendName();
+				SelectedFriend->MessageType = ChatItemSelected->GetMessageType();
+				SelectedFriend->ViewModel = nullptr;
+				SelectedFriend->SelectedMessage = ChatItemSelected;
+				bHasActionPending = true;
+				bAllowJoinGame = FFriendsAndChatManager::Get()->IsInJoinableGameSession();
+			}
 		}
 	}
 
@@ -234,15 +234,7 @@ public:
 				{
 					if (SelectedFriend.IsValid() && SelectedFriend->UserID.IsValid())
 					{
-						TSharedPtr< FFriendChatMessage > ChatItem = MakeShareable(new FFriendChatMessage());
-						ChatItem->FromName = SelectedFriend->FriendName;
-						ChatItem->Message = NewMessage;
-						ChatItem->MessageType = EChatMessageType::Whisper;
-						ChatItem->MessageTimeText = FText::AsTime(FDateTime::UtcNow());
-						ChatItem->bIsFromSelf = true;
-						ChatItem->SenderId = SelectedFriend->UserID;
-
-						if (MessageManager.Pin()->SendPrivateMessage(*SelectedFriend->UserID.Get(), ChatItem))
+						if (MessageManager.Pin()->SendPrivateMessage(SelectedFriend->UserID, SelectedFriend->FriendName, NewMessage))
 						{
 							bSuccess = true;
 						}
@@ -341,7 +333,7 @@ public:
 		{
 			return SelectedFriend->ViewModel->HasChatAction();
 		}
-		return true;
+		return false;
 	}
 
 	virtual void SetCaptureFocus(bool bInCaptureFocus) override
@@ -393,12 +385,14 @@ private:
 
 		for( const auto& Message : MessageManager.Pin()->GetMessageList())
 		{
-			HandleMessageReceived(Message.ToSharedRef());
+			HandleMessageReceived(Message);
 		}
 	}
 
 	void FilterChatList()
 	{
+		ChatLists = MessageManager.Pin()->GetMessageList();
+
 		bool bAddedItem = true;
 		if(!IsGlobalChatEnabled())
 		{
@@ -437,17 +431,12 @@ private:
 		return nullptr;
 	}
 
-	void HandleMessageReceived(const TSharedRef<FFriendChatMessage> NewMessage)
+	void HandleMessageReceived(const TSharedRef<FChatItemViewModel> NewMessage)
 	{
-		if(IsGlobalChatEnabled() || NewMessage->MessageType != EChatMessageType::Global)
+		if(IsGlobalChatEnabled() || NewMessage->GetMessageType() != EChatMessageType::Global)
 		{
-			ChatLists.Add(FChatItemViewModelFactory::Create(NewMessage, SharedThis(this)));
-			if(ChatLists.Num() > MessageDisplay)
-			{
-				ChatLists.RemoveAt(0);
-			}
+			FilterChatList();
 		}
-		FilterChatList();
 	}
 
 	void HandleChatFriendSelected(TSharedPtr<IFriendItem> ChatFriend)
@@ -483,13 +472,14 @@ private:
 			if(SelectedFriend->UserID.IsValid())
 			{
 				SelectedFriend->ViewModel = FFriendsAndChatManager::Get()->GetFriendViewModel(*SelectedFriend->UserID.Get());
+				bHasActionPending = false;
 			}
-			else if(SelectedFriend->SelectedMessage.IsValid() && SelectedFriend->SelectedMessage->MessageRef.IsValid())
+			else if(SelectedFriend->SelectedMessage.IsValid() && SelectedFriend->SelectedMessage->GetSenderID().IsValid())
 			{
-				SelectedFriend->ViewModel = FFriendsAndChatManager::Get()->GetFriendViewModel(SelectedFriend->SelectedMessage->MessageRef->GetUserId());
+				SelectedFriend->ViewModel = FFriendsAndChatManager::Get()->GetFriendViewModel(*SelectedFriend->SelectedMessage->GetSenderID().Get());
+				bHasActionPending = false;
 			}
 		}
-		bHasActionPending = false;
 	}
 
 	FChatViewModelImpl(const TSharedRef<FFriendsMessageManager>& MessageManager)
diff --git a/Engine/Source/Developer/FriendsAndChat/Private/Models/ChatViewModel.h b/Engine/Source/Developer/FriendsAndChat/Private/Models/ChatViewModel.h
--- a/Engine/Source/Developer/FriendsAndChat/Private/Models/ChatViewModel.h
+++ b/Engine/Source/Developer/FriendsAndChat/Private/Models/ChatViewModel.h
@@ -16,7 +16,7 @@ struct FSelectedFriend
 	FText FriendName;
 	EChatMessageType::Type MessageType;
 	TSharedPtr<FFriendViewModel> ViewModel;
-	TSharedPtr<FFriendChatMessage> SelectedMessage;
+	TSharedPtr<FChatItemViewModel> SelectedMessage;
 };
 
 class FChatViewModel
@@ -24,7 +24,7 @@ class FChatViewModel
 	, public IChatViewModel
 {
 public:
-	virtual TArray< TSharedRef<FChatItemViewModel > >& GetFilteredChatList() = 0;
+	virtual TArray<TSharedRef<FChatItemViewModel > >& GetFilteredChatList() = 0;
 	virtual FReply HandleSelectionChanged(TSharedRef<FChatItemViewModel> ItemSelected) = 0;
 	virtual FText GetViewGroupText() const = 0;
 	virtual FText GetChatGroupText() const = 0;
@@ -38,7 +38,7 @@ public:
 	virtual void SetChatChannel(const EChatMessageType::Type NewOption) = 0;
 	virtual void SetWhisperChannel(const TSharedPtr<FSelectedFriend> InFriend) = 0;
 	virtual void SetViewChannel(const EChatMessageType::Type NewOption) = 0;
-	virtual void SetChannelUserClicked(const TSharedRef<FFriendChatMessage> ChatItemSelected) = 0;
+	virtual void SetChannelUserClicked(const TSharedRef<FChatItemViewModel> ChatItemSelected) = 0;
 	virtual bool SendMessage(const FText NewMessage) = 0;
 	virtual void SetTimeDisplayTransparency(const float TimeTransparency) = 0;
 	virtual const float GetTimeTransparency() const = 0;
diff --git a/Engine/Source/Developer/FriendsAndChat/Private/UI/Widgets/SChatItem.cpp b/Engine/Source/Developer/FriendsAndChat/Private/UI/Widgets/SChatItem.cpp
--- a/Engine/Source/Developer/FriendsAndChat/Private/UI/Widgets/SChatItem.cpp
+++ b/Engine/Source/Developer/FriendsAndChat/Private/UI/Widgets/SChatItem.cpp
@@ -2,6 +2,7 @@
 
 #include "FriendsAndChatPrivatePCH.h"
 #include "ChatItemViewModel.h"
+#include "ChatViewModel.h"
 #include "SChatItem.h"
 
 #define LOCTEXT_NAMESPACE "SChatWindow"
@@ -10,11 +11,12 @@ class SChatItemImpl : public SChatItem
 {
 public:
 
-	void Construct(const FArguments& InArgs, const TSharedRef<FChatItemViewModel>& InViewModel)
+	void Construct(const FArguments& InArgs, const TSharedRef<FChatItemViewModel>& InViewModel, const TSharedRef<FChatViewModel>& InOwnerViewModel)
 	{
 		FriendStyle = *InArgs._FriendStyle;
 		MenuMethod = InArgs._Method;
 		this->ViewModel = InViewModel;
+		this->OwnerViewModel = InOwnerViewModel;
 
 		FText DisplayNameText = ViewModel->GetFriendNameDisplayText();
 		if(ViewModel->IsFromSelf())
@@ -27,12 +29,20 @@ public:
 			}
 		}
 
-		FFormatNamedArguments Args;
-		Args.Add(TEXT("DisplayName"), DisplayNameText);
-		Args.Add(TEXT("Message"), ViewModel->GetMessage());
-		Args.Add(TEXT("NameStyle"), FText::FromString(GetTextHyperlinkStyle()));
-		FText DisplayText = FText::Format(LOCTEXT("SChatItem_Message", "<a id=\"UserName\" style=\"{NameStyle}\">{DisplayName}</>{Message}"), Args);
-
+		FText DisplayText;
+		if(ViewModel->GetMessageType() == EChatMessageType::Party)
+		{
+			DisplayText = ViewModel->GetMessage();
+		}
+		else
+		{
+			FFormatNamedArguments Args;
+			Args.Add(TEXT("DisplayName"), DisplayNameText);
+			Args.Add(TEXT("Message"), ViewModel->GetMessage());
+			Args.Add(TEXT("NameStyle"), FText::FromString(GetTextHyperlinkStyle()));
+			DisplayText = FText::Format(LOCTEXT("SChatItem_Message", "<a id=\"UserName\" style=\"{NameStyle}\">{DisplayName}</>{Message}"), Args);
+		}
+		
 		FTextBlockStyle TextStyle = FriendStyle.TextStyle;
 		TextStyle.ColorAndOpacity = GetChannelColor();
 
@@ -59,7 +69,7 @@ public:
 				.Text(DisplayText)
 				.TextStyle(&TextStyle)
 				.DecoratorStyleSet(&FFriendsAndChatModuleStyle::Get())
-				.WrapTextAt(FriendStyle.ChatListWidth - 5)
+				.WrapTextAt(FriendStyle.ChatListWidth - 10)
 				+ SRichTextBlock::HyperlinkDecorator(TEXT( "UserName" ), this, &SChatItemImpl::HandleNameClicked)
 			]
 			+SHorizontalBox::Slot()
@@ -78,17 +88,17 @@ private:
 
 	FSlateColor GetTimeDisplayColor() const
 	{
-		if(ViewModel->UseOverrideColor())
+		if(OwnerViewModel->GetOverrideColorSet())
 		{
-			return ViewModel->GetOverrideColor();
+			return OwnerViewModel->GetFontOverrideColor();
 		}
 		else
 		{
 			switch(ViewModel->GetMessageType())
 			{
-				case EChatMessageType::Global: return FriendStyle.DefaultChatColor.CopyWithNewOpacity(ViewModel->GetFadeAmountColor()); break;
-				case EChatMessageType::Whisper: return FriendStyle.WhisplerChatColor.CopyWithNewOpacity(ViewModel->GetFadeAmountColor()); break;
-				case EChatMessageType::Party: return FriendStyle.PartyChatColor.CopyWithNewOpacity(ViewModel->GetFadeAmountColor()); break;
+				case EChatMessageType::Global: return FriendStyle.DefaultChatColor.CopyWithNewOpacity(OwnerViewModel->GetTimeTransparency()); break;
+				case EChatMessageType::Whisper: return FriendStyle.WhisplerChatColor.CopyWithNewOpacity(OwnerViewModel->GetTimeTransparency()); break;
+				case EChatMessageType::Party: return FriendStyle.PartyChatColor.CopyWithNewOpacity(OwnerViewModel->GetTimeTransparency()); break;
 				default: return FLinearColor::Gray;
 			}
 		}
@@ -96,9 +106,9 @@ private:
 
 	FSlateColor GetChannelColor () const
 	{
-		if(ViewModel->UseOverrideColor())
+		if(OwnerViewModel->GetOverrideColorSet())
 		{
-			return ViewModel->GetOverrideColor();
+			return OwnerViewModel->GetFontOverrideColor();
 		}
 		else
 		{
@@ -116,7 +126,7 @@ private:
 
 	const FSlateBrush* GetChatIcon() const
 	{
-		if(ViewModel->UseOverrideColor())
+		if(OwnerViewModel->GetOverrideColorSet())
 		{
 			return nullptr;
 		}
@@ -147,7 +157,7 @@ private:
 
 	void HandleNameClicked( const FSlateHyperlinkRun::FMetadata& Metadata )
 	{
-		ViewModel->FriendNameSelected();
+		OwnerViewModel->SetChannelUserClicked(ViewModel.ToSharedRef());
 	}
 
 private:
@@ -155,6 +165,9 @@ private:
 	// Holds the Friends List view model
 	TSharedPtr<FChatItemViewModel> ViewModel;
 
+	// Holds the owner view model
+	TSharedPtr<FChatViewModel> OwnerViewModel;
+
 	TSharedPtr<SBorder> FriendItemBorder;
 
 	/** Holds the style to use when making the widget. */
diff --git a/Engine/Source/Developer/FriendsAndChat/Private/UI/Widgets/SChatItem.h b/Engine/Source/Developer/FriendsAndChat/Private/UI/Widgets/SChatItem.h
--- a/Engine/Source/Developer/FriendsAndChat/Private/UI/Widgets/SChatItem.h
+++ b/Engine/Source/Developer/FriendsAndChat/Private/UI/Widgets/SChatItem.h
@@ -16,5 +16,5 @@ public:
 	 *
 	 * @param InArgs The Slate argument list.
 	 */
-	virtual void Construct(const FArguments& InArgs, const TSharedRef<class FChatItemViewModel>& ViewModel) = 0;
+	virtual void Construct(const FArguments& InArgs, const TSharedRef<class FChatItemViewModel>& ViewModel, const TSharedRef<class FChatViewModel>& OwnerViewModel) = 0;
 };
diff --git a/Engine/Source/Developer/FriendsAndChat/Private/UI/Widgets/SChatWindow.cpp b/Engine/Source/Developer/FriendsAndChat/Private/UI/Widgets/SChatWindow.cpp
--- a/Engine/Source/Developer/FriendsAndChat/Private/UI/Widgets/SChatWindow.cpp
+++ b/Engine/Source/Developer/FriendsAndChat/Private/UI/Widgets/SChatWindow.cpp
@@ -284,7 +284,7 @@ private:
 			.OnClicked(ViewModel.Get(), &FChatViewModel::HandleSelectionChanged, ChatMessage)
 			.VAlign(VAlign_Center)
 			[
-				SNew(SChatItem, ChatMessage)
+				SNew(SChatItem, ChatMessage, ViewModel.ToSharedRef())
 				.FriendStyle(&FriendStyle)
 				.Method(MenuMethod)
 			]
