#FriendsAndChat - Fix messages pushing date in the launcher, also fixed launcher restoring from minimized, and bring to front when social icon clicked

diff --git a/Engine/Source/Developer/FriendsAndChat/Private/FriendsAndChatManager.cpp b/Engine/Source/Developer/FriendsAndChat/Private/FriendsAndChatManager.cpp
--- a/Engine/Source/Developer/FriendsAndChat/Private/FriendsAndChatManager.cpp
+++ b/Engine/Source/Developer/FriendsAndChat/Private/FriendsAndChatManager.cpp
@@ -257,7 +257,7 @@ void FFriendsAndChatManager::CreateFriendsListWindow(const FFriendsAndChatStyle*
 		BuildFriendsUI();
 		FriendWindow = FSlateApplication::Get().AddWindow(FriendWindow.ToSharedRef());
 	}
-	else if(!FriendWindow->IsWindowMinimized())
+	else if(FriendWindow->IsWindowMinimized())
 	{
 		FriendWindow->Restore();
 		BuildFriendsUI();
@@ -265,6 +265,8 @@ void FFriendsAndChatManager::CreateFriendsListWindow(const FFriendsAndChatStyle*
 
 	GenerateChatWindow();
 
+	FriendWindow->BringToFront(true);
+
 	// Clear notifications
 	OnFriendsNotification().Broadcast(false);
 }
@@ -366,7 +368,7 @@ TSharedPtr<IChatViewModel> FFriendsAndChatManager::GetChatViewModel()
 
 void FFriendsAndChatManager::GenerateChatWindow()
 {
-	const FVector2D DEFAULT_WINDOW_SIZE = FVector2D(400, 300);
+	const FVector2D DEFAULT_WINDOW_SIZE = FVector2D(420, 500);
 
 	check(MessageManager.IsValid());
 	bCreateChatWindow = true;
@@ -393,6 +395,7 @@ void FFriendsAndChatManager::GenerateChatWindow()
 		ChatWindow->Restore();
 		SetChatWindowContents();
 	}
+	ChatWindow->BringToFront(true);
 }
 
 void FFriendsAndChatManager::SetChatFriend( TSharedPtr< IFriendItem > FriendItem )
