#FriendsAndChat - update to hide ingame rich text messages

diff --git a/Engine/Source/Developer/FriendsAndChat/Private/UI/Widgets/SChatItem.cpp b/Engine/Source/Developer/FriendsAndChat/Private/UI/Widgets/SChatItem.cpp
--- a/Engine/Source/Developer/FriendsAndChat/Private/UI/Widgets/SChatItem.cpp
+++ b/Engine/Source/Developer/FriendsAndChat/Private/UI/Widgets/SChatItem.cpp
@@ -66,6 +66,7 @@ public:
 			.MaxWidth(FriendStyle.ChatListWidth)
 			[
 				SNew( SRichTextBlock )
+				.Visibility(this, &SChatItemImpl::GetTextVisibility)
 				.Text(DisplayText)
 				.TextStyle(&TextStyle)
 				.DecoratorStyleSet(&FFriendsAndChatModuleStyle::Get())
@@ -104,6 +105,11 @@ private:
 		}
 	}
 
+	EVisibility GetTextVisibility() const
+	{
+		return OwnerViewModel->GetOverrideColorSet() ? EVisibility::Hidden : EVisibility::Visible;
+	}
+
 	FSlateColor GetChannelColor () const
 	{
 		if(OwnerViewModel->GetOverrideColorSet())
