Replaced use of SMenuAnchor::UseCurrentWindow with EPopupMethod::UseCurrentWindow to fix linking errors in Clang

diff --git a/Engine/Source/Developer/FriendsAndChat/Private/FriendsAndChatManager.cpp b/Engine/Source/Developer/FriendsAndChat/Private/FriendsAndChatManager.cpp
--- a/Engine/Source/Developer/FriendsAndChat/Private/FriendsAndChatManager.cpp
+++ b/Engine/Source/Developer/FriendsAndChat/Private/FriendsAndChatManager.cpp
@@ -351,7 +351,7 @@ TSharedPtr< SWidget > FFriendsAndChatManager::GenerateChatWidget(const FFriendsA
 	Style = *InStyle;
 	SAssignNew(ChatWidget, SChatWindow, ChatViewModel)
 	.FriendStyle(&Style)
-	.Method(SMenuAnchor::UseCurrentWindow);
+	.Method(EPopupMethod::UseCurrentWindow);
 	return ChatWidget;
 }
 
