#FriendsAndChat - Fix non-Unity build failure for the Launcher

diff --git a/Engine/Source/Developer/FriendsAndChat/Private/UI/Widgets/SFriendsContainer.cpp b/Engine/Source/Developer/FriendsAndChat/Private/UI/Widgets/SFriendsContainer.cpp
--- a/Engine/Source/Developer/FriendsAndChat/Private/UI/Widgets/SFriendsContainer.cpp
+++ b/Engine/Source/Developer/FriendsAndChat/Private/UI/Widgets/SFriendsContainer.cpp
@@ -2,6 +2,7 @@
 
 #include "FriendsAndChatPrivatePCH.h"
 #include "FriendsViewModel.h"
+#include "FriendListViewModel.h"
 #include "SFriendsList.h"
 #include "SFriendsListContainer.h"
 #include "SFriendsContainer.h"
@@ -215,7 +216,7 @@ private:
 	TSharedPtr<FFriendsViewModel> ViewModel;
 
 	// Holds the Friends Sub-List view models
-	TArray<TSharedPtr<class FFriendListViewModel>> ListViewModels;
+	TArray<TSharedPtr<FFriendListViewModel>> ListViewModels;
 
 	// Holds the menu method - Full screen requires use owning window or crashes.
 	EPopupMethod MenuMethod;
