Delegate comparisons deprecated, lots of other associated code deprecated, and lots of warning fixups:

* Multicast delegate Add* calls now return FDelegateHandles, and Remove* calls are now all deprecated, except for a new Remove function which takes a FDelegateHandle.
* New FConsoleManager::RegisterConsoleVariableSink_Handle and UnregisterConsoleVariableSink_Handle functions which work in terms of FConsoleVariableSinkHandle.
* Timer calls which don't take FTimerHandles are deprecated.
* FTicker::AddTicker now returns an FDelegateHandle and is removed by an overloaded Remove function.
* DEFINE_ONLINE_DELEGATE* macros now define _Handle variants of the Add/Remove functions which return/take handles.
* Various other handle-based registration changes.
* Some unity build fixes.
* Some simplification of delegate code.
* Fixes for lots of existing code to use handle-based registration and unregistration.

diff --git a/Engine/Plugins/2D/Paper2D/Source/Paper2D/Private/Rendering/PaperBatchManager.cpp b/Engine/Plugins/2D/Paper2D/Source/Paper2D/Private/Rendering/PaperBatchManager.cpp
--- a/Engine/Plugins/2D/Paper2D/Source/Paper2D/Private/Rendering/PaperBatchManager.cpp
+++ b/Engine/Plugins/2D/Paper2D/Source/Paper2D/Private/Rendering/PaperBatchManager.cpp
@@ -11,6 +11,9 @@ namespace FPaperBatchManagerNS
 
 	static FWorldDelegates::FWorldInitializationEvent::FDelegate OnWorldCreatedDelegate;
 	static FWorldDelegates::FWorldEvent::FDelegate OnWorldDestroyedDelegate;
+
+	static FDelegateHandle OnWorldCreatedDelegateHandle;
+	static FDelegateHandle OnWorldDestroyedDelegateHandle;
 }
 
 //////////////////////////////////////////////////////////////////////////
@@ -18,17 +21,17 @@ namespace FPaperBatchManagerNS
 
 void FPaperBatchManager::Initialize()
 {
-	FPaperBatchManagerNS::OnWorldCreatedDelegate = FWorldDelegates::FWorldInitializationEvent::FDelegate::CreateStatic(&FPaperBatchManager::OnWorldCreated);
-	FWorldDelegates::OnPostWorldInitialization.Add(FPaperBatchManagerNS::OnWorldCreatedDelegate);
+	FPaperBatchManagerNS::OnWorldCreatedDelegate       = FWorldDelegates::FWorldInitializationEvent::FDelegate::CreateStatic(&FPaperBatchManager::OnWorldCreated);
+	FPaperBatchManagerNS::OnWorldCreatedDelegateHandle = FWorldDelegates::OnPostWorldInitialization.Add(FPaperBatchManagerNS::OnWorldCreatedDelegate);
 
-	FPaperBatchManagerNS::OnWorldDestroyedDelegate = FWorldDelegates::FWorldEvent::FDelegate::CreateStatic(&FPaperBatchManager::OnWorldDestroyed);
-	FWorldDelegates::OnPreWorldFinishDestroy.Add(FPaperBatchManagerNS::OnWorldDestroyedDelegate);
+	FPaperBatchManagerNS::OnWorldDestroyedDelegate       = FWorldDelegates::FWorldEvent::FDelegate::CreateStatic(&FPaperBatchManager::OnWorldDestroyed);
+	FPaperBatchManagerNS::OnWorldDestroyedDelegateHandle = FWorldDelegates::OnPreWorldFinishDestroy.Add(FPaperBatchManagerNS::OnWorldDestroyedDelegate);
 }
 
 void FPaperBatchManager::Shutdown()
 {
-	FWorldDelegates::OnPostWorldInitialization.Remove(FPaperBatchManagerNS::OnWorldCreatedDelegate);
-	FWorldDelegates::OnPreWorldFinishDestroy.Remove(FPaperBatchManagerNS::OnWorldDestroyedDelegate);
+	FWorldDelegates::OnPostWorldInitialization.Remove(FPaperBatchManagerNS::OnWorldCreatedDelegateHandle);
+	FWorldDelegates::OnPreWorldFinishDestroy  .Remove(FPaperBatchManagerNS::OnWorldDestroyedDelegateHandle);
 }
 
 void FPaperBatchManager::OnWorldCreated(UWorld* World, const UWorld::InitializationValues IVS)
diff --git a/Engine/Plugins/2D/Paper2D/Source/Paper2DEditor/Private/ContentBrowserExtensions/ContentBrowserExtensions.cpp b/Engine/Plugins/2D/Paper2D/Source/Paper2DEditor/Private/ContentBrowserExtensions/ContentBrowserExtensions.cpp
--- a/Engine/Plugins/2D/Paper2D/Source/Paper2DEditor/Private/ContentBrowserExtensions/ContentBrowserExtensions.cpp
+++ b/Engine/Plugins/2D/Paper2D/Source/Paper2DEditor/Private/ContentBrowserExtensions/ContentBrowserExtensions.cpp
@@ -12,6 +12,7 @@ DEFINE_LOG_CATEGORY(LogPaperCBExtensions);
 //////////////////////////////////////////////////////////////////////////
 
 FContentBrowserMenuExtender_SelectedAssets ContentBrowserExtenderDelegate;
+FDelegateHandle ContentBrowserExtenderDelegateHandle;
 
 //////////////////////////////////////////////////////////////////////////
 // FContentBrowserSelectedAssetExtensionBase
@@ -324,12 +325,13 @@ void FPaperContentBrowserExtensions::InstallHooks()
 
 	TArray<FContentBrowserMenuExtender_SelectedAssets>& CBMenuExtenderDelegates = FPaperContentBrowserExtensions_Impl::GetExtenderDelegates();
 	CBMenuExtenderDelegates.Add(ContentBrowserExtenderDelegate);
+	ContentBrowserExtenderDelegateHandle = CBMenuExtenderDelegates.Last().GetHandle();
 }
 
 void FPaperContentBrowserExtensions::RemoveHooks()
 {
 	TArray<FContentBrowserMenuExtender_SelectedAssets>& CBMenuExtenderDelegates = FPaperContentBrowserExtensions_Impl::GetExtenderDelegates();
-	CBMenuExtenderDelegates.Remove(ContentBrowserExtenderDelegate);
+	CBMenuExtenderDelegates.RemoveAll([](const FContentBrowserMenuExtender_SelectedAssets& Delegate){ return Delegate.GetHandle() == ContentBrowserExtenderDelegateHandle; });
 }
 
 //////////////////////////////////////////////////////////////////////////
diff --git a/Engine/Plugins/2D/Paper2D/Source/Paper2DEditor/Private/Paper2DEditorModule.cpp b/Engine/Plugins/2D/Paper2D/Source/Paper2DEditor/Private/Paper2DEditorModule.cpp
--- a/Engine/Plugins/2D/Paper2D/Source/Paper2DEditor/Private/Paper2DEditorModule.cpp
+++ b/Engine/Plugins/2D/Paper2D/Source/Paper2DEditor/Private/Paper2DEditorModule.cpp
@@ -79,6 +79,7 @@ private:
 
 	TSharedPtr<IMeshPaintGeometryAdapterFactory> SpriteMeshPaintAdapterFactory;
 	FCoreUObjectDelegates::FOnObjectPropertyChanged::FDelegate OnPropertyChangedHandle;
+	FDelegateHandle OnPropertyChangedHandleDelegateHandle;
 
 public:
 	virtual void StartupModule() override
@@ -127,7 +128,7 @@ public:
 
 		// Register to be notified when properties are edited
 		OnPropertyChangedHandle = FCoreUObjectDelegates::FOnObjectPropertyChanged::FDelegate::CreateRaw(this, &FPaper2DEditor::OnPropertyChanged);
-		FCoreUObjectDelegates::OnObjectPropertyChanged.Add(OnPropertyChangedHandle);
+		OnPropertyChangedHandleDelegateHandle = FCoreUObjectDelegates::OnObjectPropertyChanged.Add(OnPropertyChangedHandle);
 
 		// Register the thumbnail renderers
 		UThumbnailManager::Get().RegisterCustomRenderer(UPaperSprite::StaticClass(), UPaperSpriteThumbnailRenderer::StaticClass());
@@ -187,7 +188,7 @@ public:
 			UThumbnailManager::Get().UnregisterCustomRenderer(UPaperFlipbook::StaticClass());
 
 			// Unregister the property modification handler
-			FCoreUObjectDelegates::OnObjectPropertyChanged.Remove(OnPropertyChangedHandle);
+			FCoreUObjectDelegates::OnObjectPropertyChanged.Remove(OnPropertyChangedHandleDelegateHandle);
 		}
 
 		// Unregister the details customization
diff --git a/Engine/Plugins/2D/Paper2D/Source/SmartSnapping/Private/SmartSnappingModule.cpp b/Engine/Plugins/2D/Paper2D/Source/SmartSnapping/Private/SmartSnappingModule.cpp
--- a/Engine/Plugins/2D/Paper2D/Source/SmartSnapping/Private/SmartSnappingModule.cpp
+++ b/Engine/Plugins/2D/Paper2D/Source/SmartSnapping/Private/SmartSnappingModule.cpp
@@ -14,6 +14,7 @@ class FSmartSnappingModule : public FDefaultModuleImpl
 public:
 	TSharedPtr<FPlanarConstraintSnapPolicy> PlanarPolicy;
 	FLevelEditorModule::FLevelEditorMenuExtender ViewMenuExtender;
+	FDelegateHandle ViewMenuExtenderHandle;
 
 	virtual void StartupModule() override
 	{
@@ -31,7 +32,9 @@ public:
 
 
 				FLevelEditorModule& LevelEditor = FModuleManager::LoadModuleChecked<FLevelEditorModule>(TEXT("LevelEditor"));
-				LevelEditor.GetAllLevelEditorToolbarViewMenuExtenders().Add(ViewMenuExtender);
+				auto& MenuExtenders = LevelEditor.GetAllLevelEditorToolbarViewMenuExtenders();
+				MenuExtenders.Add(ViewMenuExtender);
+				ViewMenuExtenderHandle = MenuExtenders.Last().GetHandle();
 			}
 		}
 	}
@@ -43,7 +46,7 @@ public:
 			// Unregister the level editor extensions
 			{
 				FLevelEditorModule& LevelEditor = FModuleManager::GetModuleChecked<FLevelEditorModule>( TEXT("LevelEditor") );
-				LevelEditor.GetAllLevelEditorToolbarViewMenuExtenders().Remove(ViewMenuExtender);
+				LevelEditor.GetAllLevelEditorToolbarViewMenuExtenders().RemoveAll([=](const FLevelEditorModule::FLevelEditorMenuExtender& Extender) { return Extender.GetHandle() == ViewMenuExtenderHandle; });
 			}
 			
 			//
diff --git a/Engine/Plugins/Developer/PerforceSourceControl/Source/PerforceSourceControl/Private/PerforceSourceControlProvider.cpp b/Engine/Plugins/Developer/PerforceSourceControl/Source/PerforceSourceControl/Private/PerforceSourceControlProvider.cpp
--- a/Engine/Plugins/Developer/PerforceSourceControl/Source/PerforceSourceControl/Private/PerforceSourceControlProvider.cpp
+++ b/Engine/Plugins/Developer/PerforceSourceControl/Source/PerforceSourceControl/Private/PerforceSourceControlProvider.cpp
@@ -212,7 +212,17 @@ void FPerforceSourceControlProvider::RegisterSourceControlStateChanged( const FS
 
 void FPerforceSourceControlProvider::UnregisterSourceControlStateChanged( const FSourceControlStateChanged::FDelegate& SourceControlStateChanged )
 {
-	OnSourceControlStateChanged.Remove( SourceControlStateChanged );
+	OnSourceControlStateChanged.DEPRECATED_Remove( SourceControlStateChanged );
+}
+
+FDelegateHandle FPerforceSourceControlProvider::RegisterSourceControlStateChanged_Handle( const FSourceControlStateChanged::FDelegate& SourceControlStateChanged )
+{
+	return OnSourceControlStateChanged.Add( SourceControlStateChanged );
+}
+
+void FPerforceSourceControlProvider::UnregisterSourceControlStateChanged_Handle( FDelegateHandle Handle )
+{
+	OnSourceControlStateChanged.Remove( Handle );
 }
 
 ECommandResult::Type FPerforceSourceControlProvider::Execute( const TSharedRef<ISourceControlOperation, ESPMode::ThreadSafe>& InOperation, const TArray<FString>& InFiles, EConcurrency::Type InConcurrency, const FSourceControlOperationComplete& InOperationCompleteDelegate )
diff --git a/Engine/Plugins/Developer/PerforceSourceControl/Source/PerforceSourceControl/Private/PerforceSourceControlProvider.h b/Engine/Plugins/Developer/PerforceSourceControl/Source/PerforceSourceControl/Private/PerforceSourceControlProvider.h
--- a/Engine/Plugins/Developer/PerforceSourceControl/Source/PerforceSourceControl/Private/PerforceSourceControlProvider.h
+++ b/Engine/Plugins/Developer/PerforceSourceControl/Source/PerforceSourceControl/Private/PerforceSourceControlProvider.h
@@ -8,6 +8,13 @@
 
 DECLARE_DELEGATE_RetVal(FPerforceSourceControlWorkerRef, FGetPerforceSourceControlWorker)
 
+// We disable the deprecation warnings here because otherwise it'll complain about us
+// implementing RegisterSourceControlStateChanged and UnregisterSourceControlStateChanged.  We know
+// that, but we only want it to complain if *others* implement or call these functions.
+//
+// These macros should be removed when those functions are removed.
+PRAGMA_DISABLE_DEPRECATION_WARNINGS
+
 class FPerforceSourceControlProvider : public ISourceControlProvider
 {
 public:
@@ -32,6 +39,8 @@ public:
 	virtual ECommandResult::Type GetState( const TArray<FString>& InFiles, TArray< TSharedRef<ISourceControlState, ESPMode::ThreadSafe> >& OutState, EStateCacheUsage::Type InStateCacheUsage ) override;
 	virtual void RegisterSourceControlStateChanged( const FSourceControlStateChanged::FDelegate& SourceControlStateChanged ) override;
 	virtual void UnregisterSourceControlStateChanged( const FSourceControlStateChanged::FDelegate& SourceControlStateChanged ) override;
+	virtual FDelegateHandle RegisterSourceControlStateChanged_Handle( const FSourceControlStateChanged::FDelegate& SourceControlStateChanged ) override;
+	virtual void UnregisterSourceControlStateChanged_Handle( FDelegateHandle Handle ) override;
 	virtual ECommandResult::Type Execute( const TSharedRef<ISourceControlOperation, ESPMode::ThreadSafe>& InOperation, const TArray<FString>& InFiles, EConcurrency::Type InConcurrency = EConcurrency::Synchronous, const FSourceControlOperationComplete& InOperationCompleteDelegate = FSourceControlOperationComplete() ) override;
 	virtual bool CanCancelOperation( const TSharedRef<ISourceControlOperation, ESPMode::ThreadSafe>& InOperation ) const override;
 	virtual void CancelOperation( const TSharedRef<ISourceControlOperation, ESPMode::ThreadSafe>& InOperation ) override;
@@ -144,4 +153,6 @@ private:
 
 	/** For notifying when the source control states in the cache have changed */
 	FSourceControlStateChanged OnSourceControlStateChanged;
-};
\ No newline at end of file
+};
+
+PRAGMA_ENABLE_DEPRECATION_WARNINGS
diff --git a/Engine/Plugins/Developer/VisualStudioSourceCodeAccess/Source/VisualStudioSourceCodeAccess/Private/VisualStudioSourceCodeAccessor.cpp b/Engine/Plugins/Developer/VisualStudioSourceCodeAccess/Source/VisualStudioSourceCodeAccess/Private/VisualStudioSourceCodeAccessor.cpp
--- a/Engine/Plugins/Developer/VisualStudioSourceCodeAccess/Source/VisualStudioSourceCodeAccess/Private/VisualStudioSourceCodeAccessor.cpp
+++ b/Engine/Plugins/Developer/VisualStudioSourceCodeAccess/Source/VisualStudioSourceCodeAccess/Private/VisualStudioSourceCodeAccessor.cpp
@@ -71,7 +71,7 @@ void FVisualStudioSourceCodeAccessor::Startup()
 
 #if WITH_EDITOR
 	// Setup compilation for saving all VS documents upon compilation start
-	IHotReloadModule::Get().OnModuleCompilerStarted().AddStatic( &SaveVisualStudioDocuments );
+	SaveVisualStudioDocumentsDelegateHandle = IHotReloadModule::Get().OnModuleCompilerStarted().AddStatic( &SaveVisualStudioDocuments );
 #endif
 
 	// Cache this so we don't have to do it on a background thread
@@ -91,7 +91,7 @@ void FVisualStudioSourceCodeAccessor::Shutdown()
 	// Unregister the hot-reload callback
 	if(IHotReloadModule::IsAvailable())
 	{
-		IHotReloadModule::Get().OnModuleCompilerStarted().RemoveStatic( &SaveVisualStudioDocuments );
+		IHotReloadModule::Get().OnModuleCompilerStarted().Remove( SaveVisualStudioDocumentsDelegateHandle );
 	}
 #endif
 }
diff --git a/Engine/Plugins/Developer/VisualStudioSourceCodeAccess/Source/VisualStudioSourceCodeAccess/Private/VisualStudioSourceCodeAccessor.h b/Engine/Plugins/Developer/VisualStudioSourceCodeAccess/Source/VisualStudioSourceCodeAccess/Private/VisualStudioSourceCodeAccessor.h
--- a/Engine/Plugins/Developer/VisualStudioSourceCodeAccess/Source/VisualStudioSourceCodeAccess/Private/VisualStudioSourceCodeAccessor.h
+++ b/Engine/Plugins/Developer/VisualStudioSourceCodeAccess/Source/VisualStudioSourceCodeAccess/Private/VisualStudioSourceCodeAccessor.h
@@ -15,6 +15,10 @@ public:
 #endif
 	};
 
+	FVisualStudioSourceCodeAccessor()
+	{
+	}
+
 	/** Initialise internal systems, register delegates etc. */
 	void Startup();
 
@@ -131,4 +135,8 @@ private:
 
 	/** Accessor for SolutionPath. Will try to update it when called from the game thread, otherwise will use the cached value */
 	FString GetSolutionPath() const;
+
+#if WITH_EDITOR
+	FDelegateHandle SaveVisualStudioDocumentsDelegateHandle;
+#endif
 };
diff --git a/Engine/Plugins/Editor/EpicSurvey/Source/EpicSurvey/Private/EpicSurvey.cpp b/Engine/Plugins/Editor/EpicSurvey/Source/EpicSurvey/Private/EpicSurvey.cpp
--- a/Engine/Plugins/Editor/EpicSurvey/Source/EpicSurvey/Private/EpicSurvey.cpp
+++ b/Engine/Plugins/Editor/EpicSurvey/Source/EpicSurvey/Private/EpicSurvey.cpp
@@ -118,7 +118,7 @@ void FEpicSurvey::Initialize()
 			{
 				NotificationDelegate.BindRaw( this, &FEpicSurvey::DisplayNotification );
 
-				GEditor->GetTimerManager()->SetTimer( NotificationDelegate, SurveyNotificationDelayTime, false );
+				GEditor->GetTimerManager()->SetTimer( DisplayNotificationTimerHandle, NotificationDelegate, SurveyNotificationDelayTime, false );
 			}
 		}
 
@@ -210,8 +210,8 @@ void FEpicSurvey::InitializeTitleCloud()
 
 	if (TitleCloud.IsValid())
 	{
-		TitleCloud->AddOnEnumerateFilesCompleteDelegate(FOnEnumerateFilesCompleteDelegate::CreateSP(this, &FEpicSurvey::OnEnumerateFilesComplete));
-		TitleCloud->AddOnReadFileCompleteDelegate(FOnReadFileCompleteDelegate::CreateSP(this, &FEpicSurvey::OnReadFileComplete));
+		TitleCloud->AddOnEnumerateFilesCompleteDelegate_Handle(FOnEnumerateFilesCompleteDelegate::CreateSP(this, &FEpicSurvey::OnEnumerateFilesComplete));
+		TitleCloud->AddOnReadFileCompleteDelegate_Handle(FOnReadFileCompleteDelegate::CreateSP(this, &FEpicSurvey::OnReadFileComplete));
 		LoadSurveys();
 	}
 }
@@ -450,14 +450,18 @@ void FEpicSurvey::CancelSurveyNotification()
 
 void FEpicSurvey::StartPulseSurveyIcon()
 {
-	GEditor->GetTimerManager()->SetTimer( EndPulseSurveyIconDelegate, SurveyPulseDuration, false );
+	FTimerHandle DummyHandle;
+	GEditor->GetTimerManager()->SetTimer( DummyHandle, EndPulseSurveyIconDelegate, SurveyPulseDuration, false );
+
 	bSurveyIconPulsing = true;
 }
 
 void FEpicSurvey::EndPulseSurveyIcon()
 {
 	bSurveyIconPulsing = false;
-	GEditor->GetTimerManager()->SetTimer( StartPulseSurveyIconDelegate, SurveyPulseTimeInterval, false );
+
+	FTimerHandle DummyHandle;
+	GEditor->GetTimerManager()->SetTimer( DummyHandle, StartPulseSurveyIconDelegate, SurveyPulseTimeInterval, false );
 }
 
 bool FEpicSurvey::PromptSurvey( const FGuid& SurveyIdentifier )
diff --git a/Engine/Plugins/Editor/EpicSurvey/Source/EpicSurvey/Private/EpicSurvey.h b/Engine/Plugins/Editor/EpicSurvey/Source/EpicSurvey/Private/EpicSurvey.h
--- a/Engine/Plugins/Editor/EpicSurvey/Source/EpicSurvey/Private/EpicSurvey.h
+++ b/Engine/Plugins/Editor/EpicSurvey/Source/EpicSurvey/Private/EpicSurvey.h
@@ -145,4 +145,7 @@ private:
 
 	/** The specification of the culture we are using to display surveys */
 	ECultureSpecification CurrentCulture;
+
+private:
+	FTimerHandle DisplayNotificationTimerHandle;
 };
\ No newline at end of file
diff --git a/Engine/Plugins/Experimental/LiveEditor/Source/LiveEditor/Private/LiveEditor.cpp b/Engine/Plugins/Experimental/LiveEditor/Source/LiveEditor/Private/LiveEditor.cpp
--- a/Engine/Plugins/Experimental/LiveEditor/Source/LiveEditor/Private/LiveEditor.cpp
+++ b/Engine/Plugins/Experimental/LiveEditor/Source/LiveEditor/Private/LiveEditor.cpp
@@ -156,6 +156,7 @@ private:
 	static TSharedRef<SDockTab> SpawnDeviceSetupWindow(const FSpawnTabArgs& Args);
 
 	FWorkflowApplicationModeExtender BlueprintEditorExtenderDelegate;
+	FDelegateHandle                  BlueprintEditorExtenderDelegateHandle;
 	TSharedPtr<FLiveEditorNodeInjector> NodeInjector;
 	TSharedPtr<FLiveEditorPinFactory> PinFactory;
 	FLiveEditorObjectCreateListener *ObjectCreationListener;
@@ -194,7 +195,9 @@ DECLARE_DELEGATE_RetVal_OneParam(bool, FOnBeginBlueprintContextMenuCreated, FGra
 void FLiveEditor::InstallHooks()
 {
 	BlueprintEditorExtenderDelegate = FWorkflowApplicationModeExtender::CreateStatic(&FBlueprintEditorExtenderForLiveEditor::OnModeCreated);
-	FWorkflowCentricApplication::GetModeExtenderList().Add(BlueprintEditorExtenderDelegate);
+	auto& ExtenderList = FWorkflowCentricApplication::GetModeExtenderList();
+	ExtenderList.Add(BlueprintEditorExtenderDelegate);
+	BlueprintEditorExtenderDelegateHandle = ExtenderList.Last().GetHandle();
 
 	NodeInjector = MakeShareable(new FLiveEditorNodeInjector());
 	NodeInjector->InstallHooks();
@@ -216,7 +219,7 @@ void FLiveEditor::InstallHooks()
 
 void FLiveEditor::RemoveHooks()
 {
-	FWorkflowCentricApplication::GetModeExtenderList().Remove(BlueprintEditorExtenderDelegate);
+	FWorkflowCentricApplication::GetModeExtenderList().RemoveAll([=](const FWorkflowApplicationModeExtender& Extender){ return BlueprintEditorExtenderDelegateHandle == Extender.GetHandle(); });
 
 	FEdGraphUtilities::UnregisterVisualPinFactory(PinFactory);
 
diff --git a/Engine/Plugins/Slate/SlateRemote/Source/SlateRemote/Private/Server/SlateRemoteServer.cpp b/Engine/Plugins/Slate/SlateRemote/Source/SlateRemote/Private/Server/SlateRemoteServer.cpp
--- a/Engine/Plugins/Slate/SlateRemote/Source/SlateRemote/Private/Server/SlateRemoteServer.cpp
+++ b/Engine/Plugins/Slate/SlateRemote/Source/SlateRemote/Private/Server/SlateRemoteServer.cpp
@@ -42,7 +42,7 @@ bool FSlateRemoteServer::StartServer( const FIPv4Endpoint& ServerEndpoint )
 	}
 
 	TickDelegate = FTickerDelegate::CreateRaw(this, &FSlateRemoteServer::HandleTicker);
-	FTicker::GetCoreTicker().AddTicker(TickDelegate, 0.0f);
+	TickDelegateHandle = FTicker::GetCoreTicker().AddTicker(TickDelegate, 0.0f);
 
 	return true;
 }
@@ -50,7 +50,7 @@ bool FSlateRemoteServer::StartServer( const FIPv4Endpoint& ServerEndpoint )
 
 void FSlateRemoteServer::StopServer( )
 {
-	FTicker::GetCoreTicker().RemoveTicker(TickDelegate);
+	FTicker::GetCoreTicker().RemoveTicker(TickDelegateHandle);
 
 	if (ServerSocket != nullptr)
 	{
diff --git a/Engine/Plugins/Slate/SlateRemote/Source/SlateRemote/Private/Server/SlateRemoteServer.h b/Engine/Plugins/Slate/SlateRemote/Source/SlateRemote/Private/Server/SlateRemoteServer.h
--- a/Engine/Plugins/Slate/SlateRemote/Source/SlateRemote/Private/Server/SlateRemoteServer.h
+++ b/Engine/Plugins/Slate/SlateRemote/Source/SlateRemote/Private/Server/SlateRemoteServer.h
@@ -103,6 +103,9 @@ private:
 	// Holds a delegate to be invoked when the server ticks.
 	FTickerDelegate TickDelegate;
 
+	// Handle to the registered TickDelegate.
+	FDelegateHandle TickDelegateHandle;
+
 	// The time since the last message was received from the Slate Remote application.
 	float TimeSinceLastPing;
 
diff --git a/Engine/Source/Developer/Android/AndroidTargetPlatform/Private/AndroidTargetPlatform.h b/Engine/Source/Developer/Android/AndroidTargetPlatform/Private/AndroidTargetPlatform.h
--- a/Engine/Source/Developer/Android/AndroidTargetPlatform/Private/AndroidTargetPlatform.h
+++ b/Engine/Source/Developer/Android/AndroidTargetPlatform/Private/AndroidTargetPlatform.h
@@ -184,7 +184,10 @@ private:
 
 	// Holds a delegate to be invoked when the widget ticks.
 	FTickerDelegate TickDelegate;
-	
+
+	// Handle to the registered TickDelegate.
+	FDelegateHandle TickDelegateHandle;
+
 	// Pointer to the device detection handler that grabs device ids in another thread
 	IAndroidDeviceDetection* DeviceDetection;
 
diff --git a/Engine/Source/Developer/Android/AndroidTargetPlatform/Private/AndroidTargetPlatform.inl b/Engine/Source/Developer/Android/AndroidTargetPlatform/Private/AndroidTargetPlatform.inl
--- a/Engine/Source/Developer/Android/AndroidTargetPlatform/Private/AndroidTargetPlatform.inl
+++ b/Engine/Source/Developer/Android/AndroidTargetPlatform/Private/AndroidTargetPlatform.inl
@@ -21,14 +21,14 @@ inline FAndroidTargetPlatform<TPlatformProperties>::FAndroidTargetPlatform( ) :
 	#endif
 
 	TickDelegate = FTickerDelegate::CreateRaw(this, &FAndroidTargetPlatform::HandleTicker);
-	FTicker::GetCoreTicker().AddTicker(TickDelegate, 4.0f);
+	TickDelegateHandle = FTicker::GetCoreTicker().AddTicker(TickDelegate, 4.0f);
 }
 
 
 template<class TPlatformProperties>
 inline FAndroidTargetPlatform<TPlatformProperties>::~FAndroidTargetPlatform()
 { 
-	 FTicker::GetCoreTicker().RemoveTicker(TickDelegate);
+	 FTicker::GetCoreTicker().RemoveTicker(TickDelegateHandle);
 }
 
 
diff --git a/Engine/Source/Developer/AssetTools/Private/AssetTypeActions/AssetTypeActions_CameraAnim.cpp b/Engine/Source/Developer/AssetTools/Private/AssetTypeActions/AssetTypeActions_CameraAnim.cpp
--- a/Engine/Source/Developer/AssetTools/Private/AssetTypeActions/AssetTypeActions_CameraAnim.cpp
+++ b/Engine/Source/Developer/AssetTools/Private/AssetTypeActions/AssetTypeActions_CameraAnim.cpp
@@ -221,7 +221,7 @@ void FAssetTypeActions_CameraAnim::OpenAssetEditor( const TArray<UObject*>& InOb
 				GEditor->OpenMatinee(PreviewMatineeActor.Get());
 
 				// install our delegate so we can clean up when finished
-				FEditorDelegates::EditorModeExit.AddSP(this, &FAssetTypeActions_CameraAnim::OnMatineeEditorClosed);
+				OnMatineeEditorClosedDelegateHandle = FEditorDelegates::EditorModeExit.AddSP(this, &FAssetTypeActions_CameraAnim::OnMatineeEditorClosed);
 			}
 			else
 			{
@@ -256,7 +256,7 @@ void FAssetTypeActions_CameraAnim::OnMatineeEditorClosed(FEdMode* InEditorMode)
 		}
 
 		// remove our delegate
-		FEditorDelegates::EditorModeExit.RemoveSP(this, &FAssetTypeActions_CameraAnim::OnMatineeEditorClosed);
+		FEditorDelegates::EditorModeExit.Remove(OnMatineeEditorClosedDelegateHandle);
 	}
 }
 
diff --git a/Engine/Source/Developer/AssetTools/Private/AssetTypeActions/AssetTypeActions_CameraAnim.h b/Engine/Source/Developer/AssetTools/Private/AssetTypeActions/AssetTypeActions_CameraAnim.h
--- a/Engine/Source/Developer/AssetTools/Private/AssetTypeActions/AssetTypeActions_CameraAnim.h
+++ b/Engine/Source/Developer/AssetTools/Private/AssetTypeActions/AssetTypeActions_CameraAnim.h
@@ -43,4 +43,7 @@ private:
 
 	/** The pawn we we will use for previewing the camera anim */
 	TWeakObjectPtr<class APawn> PreviewPawn;
+
+	/** Handle to the registered OnMatineeEditorClosed delegate */
+	FDelegateHandle OnMatineeEditorClosedDelegateHandle;
 };
\ No newline at end of file
diff --git a/Engine/Source/Developer/DirectoryWatcher/Private/Linux/DirectoryWatchRequestLinux.cpp b/Engine/Source/Developer/DirectoryWatcher/Private/Linux/DirectoryWatchRequestLinux.cpp
--- a/Engine/Source/Developer/DirectoryWatcher/Private/Linux/DirectoryWatchRequestLinux.cpp
+++ b/Engine/Source/Developer/DirectoryWatcher/Private/Linux/DirectoryWatchRequestLinux.cpp
@@ -77,20 +77,29 @@ bool FDirectoryWatchRequestLinux::Init(const FString& InDirectory)
 	return true;
 }
 
-void FDirectoryWatchRequestLinux::AddDelegate(const IDirectoryWatcher::FDirectoryChanged& InDelegate)
+FDelegateHandle FDirectoryWatchRequestLinux::AddDelegate(const IDirectoryWatcher::FDirectoryChanged& InDelegate)
 {
 	Delegates.Add(InDelegate);
+	return Delegates.Last().GetHandle();
 }
 
 bool FDirectoryWatchRequestLinux::RemoveDelegate(const IDirectoryWatcher::FDirectoryChanged& InDelegate)
 {
-	if (Delegates.Contains(InDelegate))
-	{
-		Delegates.Remove(InDelegate);
-		return true;
-	}
+	return DEPRECATED_RemoveDelegate(InDelegate);
+}
 
-	return false;
+bool FDirectoryWatchRequestLinux::DEPRECATED_RemoveDelegate(const IDirectoryWatcher::FDirectoryChanged& InDelegate)
+{
+	return Delegates.RemoveAll([&](const IDirectoryWatcher::FDirectoryChanged& Delegate) {
+		return Delegate.DEPRECATED_Compare(InDelegate);
+	}) != 0;
+}
+
+bool FDirectoryWatchRequestLinux::RemoveDelegate(FDelegateHandle InHandle)
+{
+	return Delegates.RemoveAll([=](const IDirectoryWatcher::FDirectoryChanged& Delegate) {
+		return Delegate.GetHandle() == InHandle;
+	}) != 0;
 }
 
 bool FDirectoryWatchRequestLinux::HasDelegates() const
diff --git a/Engine/Source/Developer/DirectoryWatcher/Private/Linux/DirectoryWatchRequestLinux.h b/Engine/Source/Developer/DirectoryWatcher/Private/Linux/DirectoryWatchRequestLinux.h
--- a/Engine/Source/Developer/DirectoryWatcher/Private/Linux/DirectoryWatchRequestLinux.h
+++ b/Engine/Source/Developer/DirectoryWatcher/Private/Linux/DirectoryWatchRequestLinux.h
@@ -20,9 +20,14 @@ public:
 	bool Init(const FString& InDirectory);
 
 	/** Adds a delegate to get fired when the directory changes */
-	void AddDelegate( const IDirectoryWatcher::FDirectoryChanged& InDelegate );
+	FDelegateHandle AddDelegate( const IDirectoryWatcher::FDirectoryChanged& InDelegate );
 	/** Removes a delegate to get fired when the directory changes */
+	DELEGATE_DEPRECATED("This overload of RemoveDelegate is deprecated, instead pass the result of AddDelegate.")
 	bool RemoveDelegate( const IDirectoryWatcher::FDirectoryChanged& InDelegate );
+	/** Same as above, but for use within other deprecated calls to prevent multiple deprecation warnings */
+	bool DEPRECATED_RemoveDelegate( const IDirectoryWatcher::FDirectoryChanged& InDelegate );
+	/** Removes a delegate to get fired when the directory changes */
+	bool RemoveDelegate( FDelegateHandle InHandle );
 	/** Returns true if this request has any delegates listening to directory changes */
 	bool HasDelegates() const;
 	/** Prepares the request for deletion */
diff --git a/Engine/Source/Developer/DirectoryWatcher/Private/Linux/DirectoryWatcherLinux.cpp b/Engine/Source/Developer/DirectoryWatcher/Private/Linux/DirectoryWatcherLinux.cpp
--- a/Engine/Source/Developer/DirectoryWatcher/Private/Linux/DirectoryWatcherLinux.cpp
+++ b/Engine/Source/Developer/DirectoryWatcher/Private/Linux/DirectoryWatcherLinux.cpp
@@ -116,6 +116,81 @@ bool FDirectoryWatcherLinux::UnregisterDirectoryChangedCallback(const FString& D
 	return false;
 }
 
+bool FDirectoryWatcherLinux::RegisterDirectoryChangedCallback_Handle( const FString& Directory, const FDirectoryChanged& InDelegate, FDelegateHandle& OutHandle )
+{
+#if WITH_DIRECTORY_WATCHER
+	
+	FDirectoryWatchRequestLinux** RequestPtr = RequestMap.Find(Directory);
+	FDirectoryWatchRequestLinux* Request = NULL;
+
+	if (RequestPtr)
+	{
+		// There should be no NULL entries in the map
+		check (*RequestPtr);
+
+		Request = *RequestPtr;
+	}
+	else
+	{
+		Request = new FDirectoryWatchRequestLinux();
+		NumRequests++;
+
+		// Begin reading directory changes
+		if (!Request->Init(Directory))
+		{
+			UE_LOG(LogDirectoryWatcher, Warning, TEXT("Failed to begin reading directory changes for %s."), *Directory);
+			delete Request;
+			NumRequests--;
+			return false;
+		}
+
+		RequestMap.Add(Directory, Request);
+	}
+
+	OutHandle = Request->AddDelegate(InDelegate);
+
+	return true;
+
+#else
+	
+	return false;
+	
+#endif // WITH_DIRECTORY_WATCHER
+}
+
+bool FDirectoryWatcherLinux::UnregisterDirectoryChangedCallback_Handle(const FString& Directory, FDelegateHandle InHandle)
+{
+#if WITH_DIRECTORY_WATCHER
+	FDirectoryWatchRequestLinux** RequestPtr = RequestMap.Find(Directory);
+
+	if (RequestPtr)
+	{
+		// There should be no NULL entries in the map
+		check (*RequestPtr);
+
+		FDirectoryWatchRequestLinux* Request = *RequestPtr;
+
+		if (Request->RemoveDelegate(InHandle))
+		{
+			if (!Request->HasDelegates())
+			{
+				// Remove from the active map and add to the pending delete list
+				RequestMap.Remove(Directory);
+				RequestsPendingDelete.AddUnique(Request);
+
+				// Signal to end the watch which will mark this request for deletion
+				Request->EndWatchRequest();
+			}
+
+			return true;
+		}
+
+	}
+#endif // WITH_DIRECTORY_WATCHER
+
+	return false;
+}
+
 void FDirectoryWatcherLinux::Tick(float DeltaSeconds)
 {
 #if WITH_DIRECTORY_WATCHER
diff --git a/Engine/Source/Developer/DirectoryWatcher/Private/Linux/DirectoryWatcherLinux.h b/Engine/Source/Developer/DirectoryWatcher/Private/Linux/DirectoryWatcherLinux.h
--- a/Engine/Source/Developer/DirectoryWatcher/Private/Linux/DirectoryWatcherLinux.h
+++ b/Engine/Source/Developer/DirectoryWatcher/Private/Linux/DirectoryWatcherLinux.h
@@ -10,6 +10,8 @@ public:
 
 	virtual bool RegisterDirectoryChangedCallback(const FString& Directory, const FDirectoryChanged& InDelegate) override;
 	virtual bool UnregisterDirectoryChangedCallback(const FString& Directory, const FDirectoryChanged& InDelegate) override;
+	virtual bool RegisterDirectoryChangedCallback_Handle(const FString& Directory, const FDirectoryChanged& InDelegate, FDelegateHandle& OutHandle) override;
+	virtual bool UnregisterDirectoryChangedCallback_Handle(const FString& Directory, FDelegateHandle InHandle) override;
 	virtual void Tick (float DeltaSeconds) override;
 
 	/** Map of directory paths to requests */
diff --git a/Engine/Source/Developer/DirectoryWatcher/Private/Mac/DirectoryWatchRequestMac.cpp b/Engine/Source/Developer/DirectoryWatcher/Private/Mac/DirectoryWatchRequestMac.cpp
--- a/Engine/Source/Developer/DirectoryWatcher/Private/Mac/DirectoryWatchRequestMac.cpp
+++ b/Engine/Source/Developer/DirectoryWatcher/Private/Mac/DirectoryWatchRequestMac.cpp
@@ -95,22 +95,29 @@ bool FDirectoryWatchRequestMac::Init(const FString& InDirectory)
 	return true;
 }
 
-void FDirectoryWatchRequestMac::AddDelegate( const IDirectoryWatcher::FDirectoryChanged& InDelegate )
+FDelegateHandle FDirectoryWatchRequestMac::AddDelegate( const IDirectoryWatcher::FDirectoryChanged& InDelegate )
 {
 	Delegates.Add(InDelegate);
+	return Delegates.Last().GetHandle();
 }
 
 bool FDirectoryWatchRequestMac::RemoveDelegate( const IDirectoryWatcher::FDirectoryChanged& InDelegate )
 {
-	if ( Delegates.Contains(InDelegate) )
-	{
-		Delegates.Remove(InDelegate);
-		return true;
-	}
-	else
-	{
-		return false;
-	}
+	return DEPRECATED_RemoveDelegate(InDelegate);
+}
+
+bool FDirectoryWatchRequestMac::DEPRECATED_RemoveDelegate( const IDirectoryWatcher::FDirectoryChanged& InDelegate )
+{
+	return Delegates.RemoveAll([&](const IDirectoryWatcher::FDirectoryChanged& Delegate) {
+		return Delegate.DEPRECATED_Compare(InDelegate);
+	}) != 0;
+}
+
+bool FDirectoryWatchRequestMac::RemoveDelegate( FDelegateHandle InHandle )
+{
+	return Delegates.RemoveAll([=](const IDirectoryWatcher::FDirectoryChanged& Delegate) {
+		return Delegate.GetHandle() == InHandle;
+	}) != 0;
 }
 
 bool FDirectoryWatchRequestMac::HasDelegates() const
diff --git a/Engine/Source/Developer/DirectoryWatcher/Private/Mac/DirectoryWatchRequestMac.h b/Engine/Source/Developer/DirectoryWatcher/Private/Mac/DirectoryWatchRequestMac.h
--- a/Engine/Source/Developer/DirectoryWatcher/Private/Mac/DirectoryWatchRequestMac.h
+++ b/Engine/Source/Developer/DirectoryWatcher/Private/Mac/DirectoryWatchRequestMac.h
@@ -15,9 +15,14 @@ public:
 	bool Init(const FString& InDirectory);
 
 	/** Adds a delegate to get fired when the directory changes */
-	void AddDelegate( const IDirectoryWatcher::FDirectoryChanged& InDelegate );
+	FDelegateHandle AddDelegate( const IDirectoryWatcher::FDirectoryChanged& InDelegate );
 	/** Removes a delegate to get fired when the directory changes */
+	DELEGATE_DEPRECATED("This overload of RemoveDelegate is deprecated, instead pass the result of AddDelegate.")
 	bool RemoveDelegate( const IDirectoryWatcher::FDirectoryChanged& InDelegate );
+	/** Same as above, but for use within other deprecated calls to prevent multiple deprecation warnings */
+	bool DEPRECATED_RemoveDelegate( const IDirectoryWatcher::FDirectoryChanged& InDelegate );
+	/** Removes a delegate to get fired when the directory changes */
+	bool RemoveDelegate( FDelegateHandle InHandle );
 	/** Returns true if this request has any delegates listening to directory changes */
 	bool HasDelegates() const;
 	/** Prepares the request for deletion */
diff --git a/Engine/Source/Developer/DirectoryWatcher/Private/Mac/DirectoryWatcherMac.cpp b/Engine/Source/Developer/DirectoryWatcher/Private/Mac/DirectoryWatcherMac.cpp
--- a/Engine/Source/Developer/DirectoryWatcher/Private/Mac/DirectoryWatcherMac.cpp
+++ b/Engine/Source/Developer/DirectoryWatcher/Private/Mac/DirectoryWatcherMac.cpp
@@ -82,7 +82,72 @@ bool FDirectoryWatcherMac::UnregisterDirectoryChangedCallback( const FString& Di
 
 		FDirectoryWatchRequestMac* Request = *RequestPtr;
 
-		if ( Request->RemoveDelegate(InDelegate) )
+		if ( Request->DEPRECATED_RemoveDelegate(InDelegate) )
+		{
+			if ( !Request->HasDelegates() )
+			{
+				// Remove from the active map and add to the pending delete list
+				RequestMap.Remove(Directory);
+				RequestsPendingDelete.AddUnique(Request);
+
+				// Signal to end the watch which will mark this request for deletion
+				Request->EndWatchRequest();
+			}
+
+			return true;
+		}
+		
+	}
+
+	return false;
+}
+
+bool FDirectoryWatcherMac::RegisterDirectoryChangedCallback_Handle( const FString& Directory, const FDirectoryChanged& InDelegate, FDelegateHandle& OutHandle )
+{
+	FDirectoryWatchRequestMac** RequestPtr = RequestMap.Find(Directory);
+	FDirectoryWatchRequestMac* Request = NULL;
+	
+	if ( RequestPtr )
+	{
+		// There should be no NULL entries in the map
+		check (*RequestPtr);
+
+		Request = *RequestPtr;
+	}
+	else
+	{
+		Request = new FDirectoryWatchRequestMac();
+		NumRequests++;
+
+		// Begin reading directory changes
+		if ( !Request->Init(Directory) )
+		{
+			UE_LOG(LogDirectoryWatcher, Warning, TEXT("Failed to begin reading directory changes for %s."), *Directory);
+			delete Request;
+			NumRequests--;
+			return false;
+		}
+
+		RequestMap.Add(Directory, Request);
+	}
+
+	OutHandle = Request->AddDelegate(InDelegate);
+
+	return true;
+}
+
+bool FDirectoryWatcherMac::UnregisterDirectoryChangedCallback_Handle( const FString& Directory, FDelegateHandle InHandle )
+{
+	FDirectoryWatchRequestMac** RequestPtr = RequestMap.Find(Directory);
+	
+	if ( RequestPtr )
+	{
+		// There should be no NULL entries in the map
+		check (*RequestPtr);
+
+		FDirectoryWatchRequestMac* Request = *RequestPtr;
+
+		if ( Request->RemoveDelegate(InHandle) )
 		{
 			if ( !Request->HasDelegates() )
 			{
diff --git a/Engine/Source/Developer/DirectoryWatcher/Private/Mac/DirectoryWatcherMac.h b/Engine/Source/Developer/DirectoryWatcher/Private/Mac/DirectoryWatcherMac.h
--- a/Engine/Source/Developer/DirectoryWatcher/Private/Mac/DirectoryWatcherMac.h
+++ b/Engine/Source/Developer/DirectoryWatcher/Private/Mac/DirectoryWatcherMac.h
@@ -11,6 +11,8 @@ public:
 
 	virtual bool RegisterDirectoryChangedCallback (const FString& Directory, const FDirectoryChanged& InDelegate) override;
 	virtual bool UnregisterDirectoryChangedCallback (const FString& Directory, const FDirectoryChanged& InDelegate) override;
+	virtual bool RegisterDirectoryChangedCallback_Handle (const FString& Directory, const FDirectoryChanged& InDelegate, FDelegateHandle& OutHandle) override;
+	virtual bool UnregisterDirectoryChangedCallback_Handle (const FString& Directory, FDelegateHandle InHandle) override;
 	virtual void Tick (float DeltaSeconds) override;
 
 	/** Map of directory paths to requests */
diff --git a/Engine/Source/Developer/DirectoryWatcher/Private/Windows/DirectoryWatchRequestWindows.cpp b/Engine/Source/Developer/DirectoryWatcher/Private/Windows/DirectoryWatchRequestWindows.cpp
--- a/Engine/Source/Developer/DirectoryWatcher/Private/Windows/DirectoryWatchRequestWindows.cpp
+++ b/Engine/Source/Developer/DirectoryWatcher/Private/Windows/DirectoryWatchRequestWindows.cpp
@@ -100,22 +100,29 @@ bool FDirectoryWatchRequestWindows::Init(const FString& InDirectory)
 	return true;
 }
 
-void FDirectoryWatchRequestWindows::AddDelegate( const IDirectoryWatcher::FDirectoryChanged& InDelegate )
+FDelegateHandle FDirectoryWatchRequestWindows::AddDelegate( const IDirectoryWatcher::FDirectoryChanged& InDelegate )
 {
 	Delegates.Add(InDelegate);
+	return Delegates.Last().GetHandle();
 }
 
 bool FDirectoryWatchRequestWindows::RemoveDelegate( const IDirectoryWatcher::FDirectoryChanged& InDelegate )
 {
-	if ( Delegates.Contains(InDelegate) )
-	{
-		Delegates.Remove(InDelegate);
-		return true;
-	}
-	else
-	{
-		return false;
-	}
+	return DEPRECATED_RemoveDelegate(InDelegate);
+}
+
+bool FDirectoryWatchRequestWindows::DEPRECATED_RemoveDelegate( const IDirectoryWatcher::FDirectoryChanged& InDelegate )
+{
+	return Delegates.RemoveAll([&](const IDirectoryWatcher::FDirectoryChanged& Delegate) {
+		return Delegate.DEPRECATED_Compare(InDelegate);
+	}) != 0;
+}
+
+bool FDirectoryWatchRequestWindows::RemoveDelegate( FDelegateHandle InHandle )
+{
+	return Delegates.RemoveAll([=](const IDirectoryWatcher::FDirectoryChanged& Delegate) {
+		return Delegate.GetHandle() == InHandle;
+	}) != 0;
 }
 
 bool FDirectoryWatchRequestWindows::HasDelegates() const
diff --git a/Engine/Source/Developer/DirectoryWatcher/Private/Windows/DirectoryWatcherWindows.cpp b/Engine/Source/Developer/DirectoryWatcher/Private/Windows/DirectoryWatcherWindows.cpp
--- a/Engine/Source/Developer/DirectoryWatcher/Private/Windows/DirectoryWatcherWindows.cpp
+++ b/Engine/Source/Developer/DirectoryWatcher/Private/Windows/DirectoryWatcherWindows.cpp
@@ -85,7 +85,73 @@ bool FDirectoryWatcherWindows::UnregisterDirectoryChangedCallback( const FString
 
 		FDirectoryWatchRequestWindows* Request = *RequestPtr;
 
-		if ( Request->RemoveDelegate(InDelegate) )
+		if ( Request->DEPRECATED_RemoveDelegate(InDelegate) )
+		{
+			if ( !Request->HasDelegates() )
+			{
+				// Remove from the active map and add to the pending delete list
+				RequestMap.Remove(Directory);
+				RequestsPendingDelete.AddUnique(Request);
+
+				// Signal to end the watch which will mark this request for deletion
+				Request->EndWatchRequest();
+			}
+
+			return true;
+		}
+		
+	}
+
+	return false;
+}
+
+bool FDirectoryWatcherWindows::RegisterDirectoryChangedCallback_Handle( const FString& Directory, const FDirectoryChanged& InDelegate, FDelegateHandle& Handle )
+{
+	FDirectoryWatchRequestWindows** RequestPtr = RequestMap.Find(Directory);
+	FDirectoryWatchRequestWindows* Request = NULL;
+	
+	if ( RequestPtr )
+	{
+		// There should be no NULL entries in the map
+		check (*RequestPtr);
+
+		Request = *RequestPtr;
+	}
+	else
+	{
+		Request = new FDirectoryWatchRequestWindows();
+		NumRequests++;
+
+		// Begin reading directory changes
+		if ( !Request->Init(Directory) )
+		{
+			uint32 Error = GetLastError();
+			UE_LOG(LogDirectoryWatcher, Warning, TEXT("Failed to begin reading directory changes for %s. Error: %d"), *Directory, Error);
+			delete Request;
+			NumRequests--;
+			return false;
+		}
+
+		RequestMap.Add(Directory, Request);
+	}
+
+	Handle = Request->AddDelegate(InDelegate);
+
+	return true;
+}
+
+bool FDirectoryWatcherWindows::UnregisterDirectoryChangedCallback_Handle( const FString& Directory, FDelegateHandle InHandle )
+{
+	FDirectoryWatchRequestWindows** RequestPtr = RequestMap.Find(Directory);
+	
+	if ( RequestPtr )
+	{
+		// There should be no NULL entries in the map
+		check (*RequestPtr);
+
+		FDirectoryWatchRequestWindows* Request = *RequestPtr;
+
+		if ( Request->RemoveDelegate(InHandle) )
 		{
 			if ( !Request->HasDelegates() )
 			{
diff --git a/Engine/Source/Developer/DirectoryWatcher/Private/Windows/DirectoryWatcherWindows.h b/Engine/Source/Developer/DirectoryWatcher/Private/Windows/DirectoryWatcherWindows.h
--- a/Engine/Source/Developer/DirectoryWatcher/Private/Windows/DirectoryWatcherWindows.h
+++ b/Engine/Source/Developer/DirectoryWatcher/Private/Windows/DirectoryWatcherWindows.h
@@ -2,6 +2,13 @@
 
 #pragma once
 
+// We disable the deprecation warnings here because otherwise it'll complain about us
+// implementing RegisterDirectoryChangedCallback and UnregisterDirectoryChangedCallback.  We know
+// that, but we only want it to complain if *others* implement or call these functions.
+//
+// These macros should be removed when those functions are removed.
+PRAGMA_DISABLE_DEPRECATION_WARNINGS
+
 class FDirectoryWatcherWindows : public IDirectoryWatcher
 {
 public:
@@ -10,6 +17,8 @@ public:
 
 	virtual bool RegisterDirectoryChangedCallback (const FString& Directory, const FDirectoryChanged& InDelegate) override;
 	virtual bool UnregisterDirectoryChangedCallback (const FString& Directory, const FDirectoryChanged& InDelegate) override;
+	virtual bool RegisterDirectoryChangedCallback_Handle (const FString& Directory, const FDirectoryChanged& InDelegate, FDelegateHandle& OutHandle) override;
+	virtual bool UnregisterDirectoryChangedCallback_Handle (const FString& Directory, FDelegateHandle InHandle) override;
 	virtual void Tick (float DeltaSeconds) override;
 
 	/** Map of directory paths to requests */
@@ -20,4 +29,6 @@ public:
 	int32 NumRequests;
 };
 
+PRAGMA_ENABLE_DEPRECATION_WARNINGS
+
 typedef FDirectoryWatcherWindows FDirectoryWatcher;
diff --git a/Engine/Source/Developer/DirectoryWatcher/Public/IDirectoryWatcher.h b/Engine/Source/Developer/DirectoryWatcher/Public/IDirectoryWatcher.h
--- a/Engine/Source/Developer/DirectoryWatcher/Public/IDirectoryWatcher.h
+++ b/Engine/Source/Developer/DirectoryWatcher/Public/IDirectoryWatcher.h
@@ -40,6 +40,7 @@ public:
 	 * @param	Directory to watch
 	 * @param	Delegate to add to our callback list
 	 */
+	DELEGATE_DEPRECATED("The RegisterDirectoryChangedCallback function is deprecated - please use the RegisterDirectoryChangedCallback_Handle function instead which returns a handle.") \
 	virtual bool RegisterDirectoryChangedCallback (const FString& Directory, const FDirectoryChanged& InDelegate) = 0;
 
 	/**
@@ -48,8 +49,26 @@ public:
 	 * @param	Directory to stop watching
 	 * @param	Delegate to remove from our callback list
 	 */
+	DELEGATE_DEPRECATED("The UnregisterDirectoryChangedCallback function is deprecated - please use the RegisterDirectoryChangedCallback_Handle function and pass its out parameter to UnregisterDirectoryChangedCallback_Handle instead.") \
 	virtual bool UnregisterDirectoryChangedCallback (const FString& Directory, const FDirectoryChanged& InDelegate) = 0;
 
+	/**
+	 * Register a callback to fire when directories are changed
+	 *
+	 * @param	Directory to watch
+	 * @param	Delegate to add to our callback list
+	 * @param	The handle to the registered delegate, if the registration was successful.
+	 */
+	virtual bool RegisterDirectoryChangedCallback_Handle (const FString& Directory, const FDirectoryChanged& InDelegate, FDelegateHandle& OutHandle) = 0;
+
+	/**
+	 * Unregisters a callback to fire when directories are changed
+	 *
+	 * @param	Directory to stop watching
+	 * @param	Handle to the delegate to remove from our callback list
+	 */
+	virtual bool UnregisterDirectoryChangedCallback_Handle (const FString& Directory, FDelegateHandle InHandle) = 0;
+
 	/**
 	 * Allows for subclasses to be ticked (by editor or other programs that need to tick the singleton)
 	 */
diff --git a/Engine/Source/Developer/FriendsAndChat/Private/FriendsAndChatManager.cpp b/Engine/Source/Developer/FriendsAndChat/Private/FriendsAndChatManager.cpp
--- a/Engine/Source/Developer/FriendsAndChat/Private/FriendsAndChatManager.cpp
+++ b/Engine/Source/Developer/FriendsAndChat/Private/FriendsAndChatManager.cpp
@@ -62,34 +62,36 @@ void FFriendsAndChatManager::Login()
 			check( FriendsInterface.IsValid() )
 
 			// Create delegates for list refreshes
-			OnQueryRecentPlayersCompleteDelegate = FOnQueryRecentPlayersCompleteDelegate::CreateRaw(this, &FFriendsAndChatManager::OnQueryRecentPlayersComplete);
-			OnFriendsListChangedDelegate = FOnFriendsChangeDelegate::CreateSP(this, &FFriendsAndChatManager::OnFriendsListChanged);
-			OnDeleteFriendCompleteDelegate = FOnDeleteFriendCompleteDelegate::CreateSP( this, &FFriendsAndChatManager::OnDeleteFriendComplete );
-			OnQueryUserIdMappingCompleteDelegate = FOnQueryUserIdMappingCompleteDelegate::CreateSP( this, &FFriendsAndChatManager::OnQueryUserIdMappingComplete );
-			OnQueryUserInfoCompleteDelegate = FOnQueryUserInfoCompleteDelegate::CreateSP( this, &FFriendsAndChatManager::OnQueryUserInfoComplete );
-			OnPresenceReceivedCompleteDelegate = FOnPresenceReceivedDelegate::CreateSP(this, &FFriendsAndChatManager::OnPresenceReceived);
-			OnPresenceUpdatedCompleteDelegate = IOnlinePresence::FOnPresenceTaskCompleteDelegate::CreateSP(this, &FFriendsAndChatManager::OnPresenceUpdated);
-			OnFriendInviteReceivedDelegate = FOnInviteReceivedDelegate::CreateSP(this, &FFriendsAndChatManager::OnFriendInviteReceived);
-			OnFriendRemovedDelegate = FOnFriendRemovedDelegate::CreateSP(this, &FFriendsAndChatManager::OnFriendRemoved);
-			OnFriendInviteRejected = FOnInviteRejectedDelegate::CreateSP(this, &FFriendsAndChatManager::OnInviteRejected);
-			OnFriendInviteAccepted = FOnInviteAcceptedDelegate::CreateSP(this, &FFriendsAndChatManager::OnInviteAccepted);
-			OnGameInviteReceivedDelegate = FOnSessionInviteReceivedDelegate::CreateSP(this, &FFriendsAndChatManager::OnGameInviteReceived);
-			OnDestroySessionCompleteDelegate = FOnDestroySessionCompleteDelegate::CreateSP(this, &FFriendsAndChatManager::OnGameDestroyed);
-
-			FriendsInterface->AddOnQueryRecentPlayersCompleteDelegate(OnQueryRecentPlayersCompleteDelegate);
-			FriendsInterface->AddOnFriendsChangeDelegate(0, OnFriendsListChangedDelegate);
-			FriendsInterface->AddOnInviteReceivedDelegate(OnFriendInviteReceivedDelegate);
-			FriendsInterface->AddOnFriendRemovedDelegate(OnFriendRemovedDelegate);
-			FriendsInterface->AddOnInviteRejectedDelegate(OnFriendInviteRejected);
-			FriendsInterface->AddOnInviteAcceptedDelegate(OnFriendInviteAccepted);
-			FriendsInterface->AddOnDeleteFriendCompleteDelegate( 0, OnDeleteFriendCompleteDelegate );
-			UserInterface->AddOnQueryUserInfoCompleteDelegate(0, OnQueryUserInfoCompleteDelegate);
-			OnlineSubMcp->GetPresenceInterface()->AddOnPresenceReceivedDelegate(OnPresenceReceivedCompleteDelegate);
-			OnlineSubMcp->GetSessionInterface()->AddOnSessionInviteReceivedDelegate(OnGameInviteReceivedDelegate);
-			OnlineSubMcp->GetSessionInterface()->AddOnDestroySessionCompleteDelegate(OnDestroySessionCompleteDelegate);
-			
+			OnQueryRecentPlayersCompleteDelegate = FOnQueryRecentPlayersCompleteDelegate           ::CreateRaw(this, &FFriendsAndChatManager::OnQueryRecentPlayersComplete);
+			OnFriendsListChangedDelegate         = FOnFriendsChangeDelegate                        ::CreateSP (this, &FFriendsAndChatManager::OnFriendsListChanged);
+			OnDeleteFriendCompleteDelegate       = FOnDeleteFriendCompleteDelegate                 ::CreateSP (this, &FFriendsAndChatManager::OnDeleteFriendComplete);
+			OnQueryUserIdMappingCompleteDelegate = FOnQueryUserIdMappingCompleteDelegate           ::CreateSP (this, &FFriendsAndChatManager::OnQueryUserIdMappingComplete);
+			OnQueryUserInfoCompleteDelegate      = FOnQueryUserInfoCompleteDelegate                ::CreateSP (this, &FFriendsAndChatManager::OnQueryUserInfoComplete);
+			OnPresenceReceivedCompleteDelegate   = FOnPresenceReceivedDelegate                     ::CreateSP (this, &FFriendsAndChatManager::OnPresenceReceived);
+			OnPresenceUpdatedCompleteDelegate    = IOnlinePresence::FOnPresenceTaskCompleteDelegate::CreateSP (this, &FFriendsAndChatManager::OnPresenceUpdated);
+			OnFriendInviteReceivedDelegate       = FOnInviteReceivedDelegate                       ::CreateSP (this, &FFriendsAndChatManager::OnFriendInviteReceived);
+			OnFriendRemovedDelegate              = FOnFriendRemovedDelegate                        ::CreateSP (this, &FFriendsAndChatManager::OnFriendRemoved);
+			OnFriendInviteRejected               = FOnInviteRejectedDelegate                       ::CreateSP (this, &FFriendsAndChatManager::OnInviteRejected);
+			OnFriendInviteAccepted               = FOnInviteAcceptedDelegate                       ::CreateSP (this, &FFriendsAndChatManager::OnInviteAccepted);
+			OnGameInviteReceivedDelegate         = FOnSessionInviteReceivedDelegate                ::CreateSP (this, &FFriendsAndChatManager::OnGameInviteReceived);
+			OnDestroySessionCompleteDelegate     = FOnDestroySessionCompleteDelegate               ::CreateSP (this, &FFriendsAndChatManager::OnGameDestroyed);
+
+			OnQueryRecentPlayersCompleteDelegateHandle = FriendsInterface->AddOnQueryRecentPlayersCompleteDelegate_Handle(   OnQueryRecentPlayersCompleteDelegate);
+			OnFriendsListChangedDelegateHandle         = FriendsInterface->AddOnFriendsChangeDelegate_Handle             (0, OnFriendsListChangedDelegate);
+			OnFriendInviteReceivedDelegateHandle       = FriendsInterface->AddOnInviteReceivedDelegate_Handle            (   OnFriendInviteReceivedDelegate);
+			OnFriendRemovedDelegateHandle              = FriendsInterface->AddOnFriendRemovedDelegate_Handle             (   OnFriendRemovedDelegate);
+			OnFriendInviteRejectedHandle               = FriendsInterface->AddOnInviteRejectedDelegate_Handle            (   OnFriendInviteRejected);
+			OnFriendInviteAcceptedHandle               = FriendsInterface->AddOnInviteAcceptedDelegate_Handle            (   OnFriendInviteAccepted);
+			OnDeleteFriendCompleteDelegateHandle       = FriendsInterface->AddOnDeleteFriendCompleteDelegate_Handle      (0, OnDeleteFriendCompleteDelegate);
+
+			OnQueryUserInfoCompleteDelegateHandle = UserInterface->AddOnQueryUserInfoCompleteDelegate_Handle(0, OnQueryUserInfoCompleteDelegate);
+
+			OnPresenceReceivedCompleteDelegateHandle = OnlineSubMcp->GetPresenceInterface()->AddOnPresenceReceivedDelegate_Handle      (OnPresenceReceivedCompleteDelegate);
+			OnGameInviteReceivedDelegateHandle       = OnlineSubMcp->GetSessionInterface ()->AddOnSessionInviteReceivedDelegate_Handle (OnGameInviteReceivedDelegate);
+			OnDestroySessionCompleteDelegateHandle   = OnlineSubMcp->GetSessionInterface ()->AddOnDestroySessionCompleteDelegate_Handle(OnDestroySessionCompleteDelegate);
+
 			FOnlinePersonaMcpPtr OnlinePersonaMcp = OnlineSubMcp->GetMcpPersonaService();
-			OnlinePersonaMcp->AddOnQueryUserIdMappingCompleteDelegate(OnQueryUserIdMappingCompleteDelegate);
+			OnQueryUserIdMappingCompleteDelegateHandle = OnlinePersonaMcp->AddOnQueryUserIdMappingCompleteDelegate_Handle(OnQueryUserIdMappingCompleteDelegate);
 
 			ManagerState = EFriendsAndManagerState::Idle;
 
@@ -101,7 +103,7 @@ void FFriendsAndChatManager::Login()
 				UpdateFriendsTickerDelegate = FTickerDelegate::CreateSP( this, &FFriendsAndChatManager::Tick );
 			}
 
-			FTicker::GetCoreTicker().AddTicker( UpdateFriendsTickerDelegate );
+			UpdateFriendsTickerDelegateHandle = FTicker::GetCoreTicker().AddTicker( UpdateFriendsTickerDelegate );
 
 			SetState(EFriendsAndManagerState::RequestFriendsListRefresh);
 			RequestRecentPlayersListRefresh();
@@ -125,30 +127,30 @@ void FFriendsAndChatManager::Logout()
 	{
 		if (OnlineSubMcp->GetFriendsInterface().IsValid())
 		{
-			OnlineSubMcp->GetFriendsInterface()->ClearOnQueryRecentPlayersCompleteDelegate(OnQueryRecentPlayersCompleteDelegate);
-			OnlineSubMcp->GetFriendsInterface()->ClearOnFriendsChangeDelegate(0, OnFriendsListChangedDelegate);
-			OnlineSubMcp->GetFriendsInterface()->ClearOnInviteReceivedDelegate(OnFriendInviteReceivedDelegate);
-			OnlineSubMcp->GetFriendsInterface()->ClearOnFriendRemovedDelegate(OnFriendRemovedDelegate);
-			OnlineSubMcp->GetFriendsInterface()->ClearOnInviteRejectedDelegate(OnFriendInviteRejected);
-			OnlineSubMcp->GetFriendsInterface()->ClearOnInviteAcceptedDelegate(OnFriendInviteAccepted);
-			OnlineSubMcp->GetFriendsInterface()->ClearOnDeleteFriendCompleteDelegate(0, OnDeleteFriendCompleteDelegate);
+			OnlineSubMcp->GetFriendsInterface()->ClearOnQueryRecentPlayersCompleteDelegate_Handle(   OnQueryRecentPlayersCompleteDelegateHandle);
+			OnlineSubMcp->GetFriendsInterface()->ClearOnFriendsChangeDelegate_Handle             (0, OnFriendsListChangedDelegateHandle);
+			OnlineSubMcp->GetFriendsInterface()->ClearOnInviteReceivedDelegate_Handle            (   OnFriendInviteReceivedDelegateHandle);
+			OnlineSubMcp->GetFriendsInterface()->ClearOnFriendRemovedDelegate_Handle             (   OnFriendRemovedDelegateHandle);
+			OnlineSubMcp->GetFriendsInterface()->ClearOnInviteRejectedDelegate_Handle            (   OnFriendInviteRejectedHandle);
+			OnlineSubMcp->GetFriendsInterface()->ClearOnInviteAcceptedDelegate_Handle            (   OnFriendInviteAcceptedHandle);
+			OnlineSubMcp->GetFriendsInterface()->ClearOnDeleteFriendCompleteDelegate_Handle      (0, OnDeleteFriendCompleteDelegateHandle);
 		}
 		if (OnlineSubMcp->GetPresenceInterface().IsValid())
 		{
-			OnlineSubMcp->GetPresenceInterface()->ClearOnPresenceReceivedDelegate(OnPresenceReceivedCompleteDelegate);
+			OnlineSubMcp->GetPresenceInterface()->ClearOnPresenceReceivedDelegate_Handle(OnPresenceReceivedCompleteDelegateHandle);
 		}
 		if (OnlineSubMcp->GetSessionInterface().IsValid())
 		{
-			OnlineSubMcp->GetSessionInterface()->ClearOnSessionInviteReceivedDelegate(OnGameInviteReceivedDelegate);
-			OnlineSubMcp->GetSessionInterface()->ClearOnDestroySessionCompleteDelegate(OnDestroySessionCompleteDelegate);
+			OnlineSubMcp->GetSessionInterface()->ClearOnSessionInviteReceivedDelegate_Handle (OnGameInviteReceivedDelegateHandle);
+			OnlineSubMcp->GetSessionInterface()->ClearOnDestroySessionCompleteDelegate_Handle(OnDestroySessionCompleteDelegateHandle);
 		}
 		if (OnlineSubMcp->GetUserInterface().IsValid())
 		{
-			OnlineSubMcp->GetUserInterface()->ClearOnQueryUserInfoCompleteDelegate(0, OnQueryUserInfoCompleteDelegate);
+			OnlineSubMcp->GetUserInterface()->ClearOnQueryUserInfoCompleteDelegate_Handle(0, OnQueryUserInfoCompleteDelegateHandle);
 		}
 		if (OnlineSubMcp->GetMcpPersonaService().IsValid())
 		{
-			OnlineSubMcp->GetMcpPersonaService()->ClearOnQueryUserIdMappingCompleteDelegate(OnQueryUserIdMappingCompleteDelegate);
+			OnlineSubMcp->GetMcpPersonaService()->ClearOnQueryUserIdMappingCompleteDelegate_Handle(OnQueryUserIdMappingCompleteDelegateHandle);
 		}
 	}
 
@@ -181,7 +183,7 @@ void FFriendsAndChatManager::Logout()
 	OnlineSubMcp = nullptr;
 	if ( UpdateFriendsTickerDelegate.IsBound() )
 	{
-		FTicker::GetCoreTicker().RemoveTicker( UpdateFriendsTickerDelegate );
+		FTicker::GetCoreTicker().RemoveTicker( UpdateFriendsTickerDelegateHandle );
 	}
 
 	SetState(EFriendsAndManagerState::OffLine);
diff --git a/Engine/Source/Developer/FriendsAndChat/Private/FriendsAndChatManager.h b/Engine/Source/Developer/FriendsAndChat/Private/FriendsAndChatManager.h
--- a/Engine/Source/Developer/FriendsAndChat/Private/FriendsAndChatManager.h
+++ b/Engine/Source/Developer/FriendsAndChat/Private/FriendsAndChatManager.h
@@ -785,4 +785,20 @@ private:
 
 	FFriendsAndChatAnalytics Analytics;
 	static TSharedPtr< FFriendsAndChatManager > SingletonInstance;
+
+	/** Handle to various registered delegates */
+	FDelegateHandle OnQueryRecentPlayersCompleteDelegateHandle;
+	FDelegateHandle OnFriendsListChangedDelegateHandle;
+	FDelegateHandle OnFriendInviteReceivedDelegateHandle;
+	FDelegateHandle OnFriendRemovedDelegateHandle;
+	FDelegateHandle OnFriendInviteRejectedHandle;
+	FDelegateHandle OnFriendInviteAcceptedHandle;
+	FDelegateHandle OnReadFriendsCompleteDelegateHandle;
+	FDelegateHandle OnDeleteFriendCompleteDelegateHandle;
+	FDelegateHandle OnQueryUserInfoCompleteDelegateHandle;
+	FDelegateHandle OnPresenceReceivedCompleteDelegateHandle;
+	FDelegateHandle OnGameInviteReceivedDelegateHandle;
+	FDelegateHandle OnDestroySessionCompleteDelegateHandle;
+	FDelegateHandle OnQueryUserIdMappingCompleteDelegateHandle;
+	FDelegateHandle UpdateFriendsTickerDelegateHandle;
 };
diff --git a/Engine/Source/Developer/FriendsAndChat/Private/FriendsMessageManager.cpp b/Engine/Source/Developer/FriendsAndChat/Private/FriendsMessageManager.cpp
--- a/Engine/Source/Developer/FriendsAndChat/Private/FriendsMessageManager.cpp
+++ b/Engine/Source/Developer/FriendsAndChat/Private/FriendsMessageManager.cpp
@@ -191,27 +191,27 @@ private:
 			IOnlineChatPtr ChatInterface = OnlineSub->GetChatInterface();
 			if (ChatInterface.IsValid())
 			{
-				OnChatRoomJoinPublicDelegate = FOnChatRoomJoinPublicDelegate::CreateSP( this, &FFriendsMessageManagerImpl::OnChatRoomJoinPublic);
-				OnChatRoomExitDelegate = FOnChatRoomExitDelegate::CreateSP( this, &FFriendsMessageManagerImpl::OnChatRoomExit);
-				OnChatRoomMemberJoinDelegate = FOnChatRoomMemberJoinDelegate::CreateSP( this, &FFriendsMessageManagerImpl::OnChatRoomMemberJoin);
-				OnChatRoomMemberExitDelegate = FOnChatRoomMemberExitDelegate::CreateSP( this, &FFriendsMessageManagerImpl::OnChatRoomMemberExit);
-				OnChatRoomMemberUpdateDelegate = FOnChatRoomMemberUpdateDelegate::CreateSP( this, &FFriendsMessageManagerImpl::OnChatRoomMemberUpdate);
-				OnChatRoomMessageReceivedDelegate = FOnChatRoomMessageReceivedDelegate::CreateSP( this, &FFriendsMessageManagerImpl::OnChatRoomMessageReceived);
-				OnChatPrivateMessageReceivedDelegate = FOnChatPrivateMessageReceivedDelegate::CreateSP( this, &FFriendsMessageManagerImpl::OnChatPrivateMessageReceived);
-
-				ChatInterface->AddOnChatRoomJoinPublicDelegate(OnChatRoomJoinPublicDelegate);
-				ChatInterface->AddOnChatRoomExitDelegate(OnChatRoomExitDelegate);
-				ChatInterface->AddOnChatRoomMemberJoinDelegate(OnChatRoomMemberJoinDelegate);
-				ChatInterface->AddOnChatRoomMemberExitDelegate(OnChatRoomMemberExitDelegate);
-				ChatInterface->AddOnChatRoomMemberUpdateDelegate(OnChatRoomMemberUpdateDelegate);
-				ChatInterface->AddOnChatRoomMessageReceivedDelegate(OnChatRoomMessageReceivedDelegate);
-				ChatInterface->AddOnChatPrivateMessageReceivedDelegate(OnChatPrivateMessageReceivedDelegate);
+				OnChatRoomJoinPublicDelegate         = FOnChatRoomJoinPublicDelegate        ::CreateSP(this, &FFriendsMessageManagerImpl::OnChatRoomJoinPublic);
+				OnChatRoomExitDelegate               = FOnChatRoomExitDelegate              ::CreateSP(this, &FFriendsMessageManagerImpl::OnChatRoomExit);
+				OnChatRoomMemberJoinDelegate         = FOnChatRoomMemberJoinDelegate        ::CreateSP(this, &FFriendsMessageManagerImpl::OnChatRoomMemberJoin);
+				OnChatRoomMemberExitDelegate         = FOnChatRoomMemberExitDelegate        ::CreateSP(this, &FFriendsMessageManagerImpl::OnChatRoomMemberExit);
+				OnChatRoomMemberUpdateDelegate       = FOnChatRoomMemberUpdateDelegate      ::CreateSP(this, &FFriendsMessageManagerImpl::OnChatRoomMemberUpdate);
+				OnChatRoomMessageReceivedDelegate    = FOnChatRoomMessageReceivedDelegate   ::CreateSP(this, &FFriendsMessageManagerImpl::OnChatRoomMessageReceived);
+				OnChatPrivateMessageReceivedDelegate = FOnChatPrivateMessageReceivedDelegate::CreateSP(this, &FFriendsMessageManagerImpl::OnChatPrivateMessageReceived);
+
+				OnChatRoomJoinPublicDelegateHandle         = ChatInterface->AddOnChatRoomJoinPublicDelegate_Handle        (OnChatRoomJoinPublicDelegate);
+				OnChatRoomExitDelegateHandle               = ChatInterface->AddOnChatRoomExitDelegate_Handle              (OnChatRoomExitDelegate);
+				OnChatRoomMemberJoinDelegateHandle         = ChatInterface->AddOnChatRoomMemberJoinDelegate_Handle        (OnChatRoomMemberJoinDelegate);
+				OnChatRoomMemberExitDelegateHandle         = ChatInterface->AddOnChatRoomMemberExitDelegate_Handle        (OnChatRoomMemberExitDelegate);
+				OnChatRoomMemberUpdateDelegateHandle       = ChatInterface->AddOnChatRoomMemberUpdateDelegate_Handle      (OnChatRoomMemberUpdateDelegate);
+				OnChatRoomMessageReceivedDelegateHandle    = ChatInterface->AddOnChatRoomMessageReceivedDelegate_Handle   (OnChatRoomMessageReceivedDelegate);
+				OnChatPrivateMessageReceivedDelegateHandle = ChatInterface->AddOnChatPrivateMessageReceivedDelegate_Handle(OnChatPrivateMessageReceivedDelegate);
 			}
 			IOnlinePresencePtr PresenceInterface = OnlineSub->GetPresenceInterface();
 			if (PresenceInterface.IsValid())
 			{
 				OnPresenceReceivedDelegate = FOnPresenceReceivedDelegate::CreateSP(this, &FFriendsMessageManagerImpl::OnPresenceReceived);
-				PresenceInterface->AddOnPresenceReceivedDelegate(OnPresenceReceivedDelegate);
+				OnPresenceReceivedDelegateHandle = PresenceInterface->AddOnPresenceReceivedDelegate_Handle(OnPresenceReceivedDelegate);
 			}
 		}
 	}
@@ -223,18 +223,18 @@ private:
 			IOnlineChatPtr ChatInterface = OnlineSub->GetChatInterface();
 			if( ChatInterface.IsValid())
 			{
-				ChatInterface->ClearOnChatRoomJoinPublicDelegate(OnChatRoomJoinPublicDelegate);
-				ChatInterface->ClearOnChatRoomExitDelegate(OnChatRoomExitDelegate);
-				ChatInterface->ClearOnChatRoomMemberJoinDelegate(OnChatRoomMemberJoinDelegate);
-				ChatInterface->ClearOnChatRoomMemberExitDelegate(OnChatRoomMemberExitDelegate);
-				ChatInterface->ClearOnChatRoomMemberUpdateDelegate(OnChatRoomMemberUpdateDelegate);
-				ChatInterface->ClearOnChatRoomMessageReceivedDelegate(OnChatRoomMessageReceivedDelegate);
-				ChatInterface->ClearOnChatPrivateMessageReceivedDelegate(OnChatPrivateMessageReceivedDelegate);
+				ChatInterface->ClearOnChatRoomJoinPublicDelegate_Handle        (OnChatRoomJoinPublicDelegateHandle);
+				ChatInterface->ClearOnChatRoomExitDelegate_Handle              (OnChatRoomExitDelegateHandle);
+				ChatInterface->ClearOnChatRoomMemberJoinDelegate_Handle        (OnChatRoomMemberJoinDelegateHandle);
+				ChatInterface->ClearOnChatRoomMemberExitDelegate_Handle        (OnChatRoomMemberExitDelegateHandle);
+				ChatInterface->ClearOnChatRoomMemberUpdateDelegate_Handle      (OnChatRoomMemberUpdateDelegateHandle);
+				ChatInterface->ClearOnChatRoomMessageReceivedDelegate_Handle   (OnChatRoomMessageReceivedDelegateHandle);
+				ChatInterface->ClearOnChatPrivateMessageReceivedDelegate_Handle(OnChatPrivateMessageReceivedDelegateHandle);
 			}
 			IOnlinePresencePtr PresenceInterface = OnlineSub->GetPresenceInterface();
 			if (PresenceInterface.IsValid())
 			{
-				PresenceInterface->ClearOnPresenceReceivedDelegate(OnPresenceReceivedDelegate);
+				PresenceInterface->ClearOnPresenceReceivedDelegate_Handle(OnPresenceReceivedDelegateHandle);
 			}
 		}
 		OnlineSub = nullptr;
@@ -450,6 +450,16 @@ private:
 	FOnChatPrivateMessageReceivedDelegate OnChatPrivateMessageReceivedDelegate;
 	FOnPresenceReceivedDelegate OnPresenceReceivedDelegate;
 
+	// Handles to the above registered delegates
+	FDelegateHandle OnChatRoomJoinPublicDelegateHandle;
+	FDelegateHandle OnChatRoomExitDelegateHandle;
+	FDelegateHandle OnChatRoomMemberJoinDelegateHandle;
+	FDelegateHandle OnChatRoomMemberExitDelegateHandle;
+	FDelegateHandle OnChatRoomMemberUpdateDelegateHandle;
+	FDelegateHandle OnChatRoomMessageReceivedDelegateHandle;
+	FDelegateHandle OnChatPrivateMessageReceivedDelegateHandle;
+	FDelegateHandle OnPresenceReceivedDelegateHandle;
+
 	// Outgoing events
 	FOnChatMessageReceivedEvent MessageReceivedEvent;
 	FOnChatPublicRoomJoinedEvent PublicRoomJoinedEvent;
diff --git a/Engine/Source/Developer/FunctionalTesting/Classes/FunctionalTestingManager.h b/Engine/Source/Developer/FunctionalTesting/Classes/FunctionalTestingManager.h
--- a/Engine/Source/Developer/FunctionalTesting/Classes/FunctionalTestingManager.h
+++ b/Engine/Source/Developer/FunctionalTesting/Classes/FunctionalTestingManager.h
@@ -84,4 +84,7 @@ protected:
 	TArray<FString> TestReproStrings;
 
 	static FAutomationTestExecutionInfo* ExecutionInfo;
+
+private:
+	FTimerHandle TriggerFirstValidTestTimerHandle;
 };
diff --git a/Engine/Source/Developer/FunctionalTesting/Private/FunctionalTestingManager.cpp b/Engine/Source/Developer/FunctionalTesting/Private/FunctionalTestingManager.cpp
--- a/Engine/Source/Developer/FunctionalTesting/Private/FunctionalTestingManager.cpp
+++ b/Engine/Source/Developer/FunctionalTesting/Private/FunctionalTestingManager.cpp
@@ -161,7 +161,7 @@ void UFunctionalTestingManager::TriggerFirstValidTest()
 	{
 		bInitialDelayApplied = true;
 		static const float WaitingTime = 0.25f;
-		World->GetTimerManager().SetTimer(this, &UFunctionalTestingManager::TriggerFirstValidTest, WaitingTime);
+		World->GetTimerManager().SetTimer(TriggerFirstValidTestTimerHandle, this, &UFunctionalTestingManager::TriggerFirstValidTest, WaitingTime);
 	}
 }
 
diff --git a/Engine/Source/Developer/HotReload/Private/HotReload.cpp b/Engine/Source/Developer/HotReload/Private/HotReload.cpp
--- a/Engine/Source/Developer/HotReload/Private/HotReload.cpp
+++ b/Engine/Source/Developer/HotReload/Private/HotReload.cpp
@@ -278,9 +278,15 @@ private:
 	/** FTicker delegate (hot-reload from IDE) */
 	FTickerDelegate TickerDelegate;
 
+	/** Handle to the registered TickerDelegate */
+	FDelegateHandle TickerDelegateHandle;
+
 	/** Callback when game binaries folder changes */
 	IDirectoryWatcher::FDirectoryChanged BinariesFolderChangedDelegate;
 
+	/** Handle to the registered delegate above */
+	FDelegateHandle BinariesFolderChangedDelegateHandle;
+
 	/** True if currently hot-reloading from editor (suppresses hot-reload from IDE) */
 	bool bIsHotReloadingFromEditor;
 	
@@ -371,14 +377,14 @@ void FHotReloadModule::StartupModule()
 
 	// Register hot-reload from IDE ticker
 	TickerDelegate = FTickerDelegate::CreateRaw(this, &FHotReloadModule::Tick);
-	FTicker::GetCoreTicker().AddTicker(TickerDelegate);
+	TickerDelegateHandle = FTicker::GetCoreTicker().AddTicker(TickerDelegate);
 
 	FModuleManager::Get().OnModulesChanged().AddRaw(this, &FHotReloadModule::ModulesChangedCallback);
 }
 
 void FHotReloadModule::ShutdownModule()
 {
-	FTicker::GetCoreTicker().RemoveTicker(TickerDelegate);
+	FTicker::GetCoreTicker().RemoveTicker(TickerDelegateHandle);
 	ShutdownHotReloadWatcher();
 }
 
@@ -922,7 +928,7 @@ void FHotReloadModule::InitHotReloadWatcher()
 		if (FPaths::DirectoryExists(BinariesPath))
 		{
 			BinariesFolderChangedDelegate = IDirectoryWatcher::FDirectoryChanged::CreateRaw(this, &FHotReloadModule::OnHotReloadBinariesChanged);
-			bDirectoryWatcherInitialized = DirectoryWatcher->RegisterDirectoryChangedCallback(BinariesPath, BinariesFolderChangedDelegate);
+			bDirectoryWatcherInitialized = DirectoryWatcher->RegisterDirectoryChangedCallback_Handle(BinariesPath, BinariesFolderChangedDelegate, BinariesFolderChangedDelegateHandle);
 		}
 	}
 }
@@ -936,7 +942,7 @@ void FHotReloadModule::ShutdownHotReloadWatcher()
 		if (DirectoryWatcher)
 		{
 			FString BinariesPath = FPaths::ConvertRelativePathToFull(FPaths::GameDir() / TEXT("Binaries") / FPlatformProcess::GetBinariesSubdirectory());
-			DirectoryWatcher->UnregisterDirectoryChangedCallback(BinariesPath, BinariesFolderChangedDelegate);
+			DirectoryWatcher->UnregisterDirectoryChangedCallback_Handle(BinariesPath, BinariesFolderChangedDelegateHandle);
 		}
 	}
 }
diff --git a/Engine/Source/Developer/LauncherServices/Private/Profiles/LauncherProfile.h b/Engine/Source/Developer/LauncherServices/Private/Profiles/LauncherProfile.h
--- a/Engine/Source/Developer/LauncherServices/Private/Profiles/LauncherProfile.h
+++ b/Engine/Source/Developer/LauncherServices/Private/Profiles/LauncherProfile.h
@@ -714,14 +714,14 @@ public:
 	{
 		if(DeployedDeviceGroup.IsValid())
 		{
-			DeployedDeviceGroup->OnDeviceAdded().RemoveRaw(this, &FLauncherProfile::OnLauncherDeviceGroupDeviceAdded);
-			DeployedDeviceGroup->OnDeviceRemoved().RemoveRaw(this, &FLauncherProfile::OnLauncherDeviceGroupDeviceRemove);
+			DeployedDeviceGroup->OnDeviceAdded().Remove(OnLauncherDeviceGroupDeviceAddedDelegateHandle);
+			DeployedDeviceGroup->OnDeviceRemoved().Remove(OnLauncherDeviceGroupDeviceRemoveDelegateHandle);
 		}
 		DeployedDeviceGroup = DeviceGroup;
 		if (DeployedDeviceGroup.IsValid())
 		{
-			DeployedDeviceGroup->OnDeviceAdded().AddRaw(this, &FLauncherProfile::OnLauncherDeviceGroupDeviceAdded);
-			DeployedDeviceGroup->OnDeviceRemoved().AddRaw(this, &FLauncherProfile::OnLauncherDeviceGroupDeviceRemove);
+			OnLauncherDeviceGroupDeviceAddedDelegateHandle   = DeployedDeviceGroup->OnDeviceAdded().AddRaw(this, &FLauncherProfile::OnLauncherDeviceGroupDeviceAdded);
+			OnLauncherDeviceGroupDeviceRemoveDelegateHandle  = DeployedDeviceGroup->OnDeviceRemoved().AddRaw(this, &FLauncherProfile::OnLauncherDeviceGroupDeviceRemove);
 			DeployedDeviceGroupId = DeployedDeviceGroup->GetId();
 		}
 		else
@@ -1146,6 +1146,10 @@ private:
 	// Holds the device group to deploy to.
 	ILauncherDeviceGroupPtr DeployedDeviceGroup;
 
+	// Delegate handles for registered DeployedDeviceGroup event handlers.
+	FDelegateHandle OnLauncherDeviceGroupDeviceAddedDelegateHandle;
+	FDelegateHandle OnLauncherDeviceGroupDeviceRemoveDelegateHandle;
+
 	// Holds the identifier of the deployed device group.
 	FGuid DeployedDeviceGroupId;
 
diff --git a/Engine/Source/Developer/LogVisualizer/Private/SVisualLogger.cpp b/Engine/Source/Developer/LogVisualizer/Private/SVisualLogger.cpp
--- a/Engine/Source/Developer/LogVisualizer/Private/SVisualLogger.cpp
+++ b/Engine/Source/Developer/LogVisualizer/Private/SVisualLogger.cpp
@@ -200,7 +200,7 @@ SVisualLogger::~SVisualLogger()
 		}
 	}
 
-	UDebugDrawService::Unregister(FDebugDrawDelegate::CreateRaw(VisualLoggerCanvasRenderer.Get(), &FVisualLoggerCanvasRenderer::DrawOnCanvas));
+	UDebugDrawService::Unregister(DrawOnCanvasDelegateHandle);
 	VisualLoggerCanvasRenderer.Reset();
 
 	FVisualLogger::Get().RemoveDevice(InternalDevice.Get());
@@ -394,7 +394,7 @@ void SVisualLogger::Construct(const FArguments& InArgs, const TSharedRef<SDockTa
 
 	VisualLoggerCanvasRenderer = MakeShareable(new FVisualLoggerCanvasRenderer(VisualLoggerInterface));
 
-	UDebugDrawService::Register(TEXT("VisLog"), FDebugDrawDelegate::CreateRaw(VisualLoggerCanvasRenderer.Get(), &FVisualLoggerCanvasRenderer::DrawOnCanvas));
+	DrawOnCanvasDelegateHandle = UDebugDrawService::Register(TEXT("VisLog"), FDebugDrawDelegate::CreateRaw(VisualLoggerCanvasRenderer.Get(), &FVisualLoggerCanvasRenderer::DrawOnCanvas));
 	//UGameplayDebuggingComponent::OnDebuggingTargetChangedDelegate.AddSP(this, &SVisualLogger::SelectionChanged);
 }
 
diff --git a/Engine/Source/Developer/LogVisualizer/Private/SVisualLogger.h b/Engine/Source/Developer/LogVisualizer/Private/SVisualLogger.h
--- a/Engine/Source/Developer/LogVisualizer/Private/SVisualLogger.h
+++ b/Engine/Source/Developer/LogVisualizer/Private/SVisualLogger.h
@@ -118,4 +118,6 @@ protected:
 	mutable TSharedPtr<SVisualLoggerView> MainView;
 	mutable TSharedPtr<SVisualLoggerLogsList> LogsList;
 	mutable TSharedPtr<SVisualLoggerStatusView> StatusView;
+
+	FDelegateHandle DrawOnCanvasDelegateHandle;
 };
diff --git a/Engine/Source/Developer/Merge/Private/SBlueprintMerge.cpp b/Engine/Source/Developer/Merge/Private/SBlueprintMerge.cpp
--- a/Engine/Source/Developer/Merge/Private/SBlueprintMerge.cpp
+++ b/Engine/Source/Developer/Merge/Private/SBlueprintMerge.cpp
@@ -486,9 +486,9 @@ void SBlueprintMerge::OnAcceptLocal()
 				// copy until then by adding it to the OnBlueprintUnloaded delegate
 				FKismetEditorUtilities::FOnBlueprintUnloaded::FDelegate OnPackageUnloaded = FKismetEditorUtilities::FOnBlueprintUnloaded::FDelegate::CreateLambda(OverwriteBlueprintFile);
 					
-				FKismetEditorUtilities::OnBlueprintUnloaded.Add(OnPackageUnloaded);
+				FDelegateHandle OnPackageUnloadedDelegateHandle = FKismetEditorUtilities::OnBlueprintUnloaded.Add(OnPackageUnloaded);
 				NewBlueprint = FKismetEditorUtilities::ReloadBlueprint(TargetBlueprint);
-				FKismetEditorUtilities::OnBlueprintUnloaded.Remove(OnPackageUnloaded);
+				FKismetEditorUtilities::OnBlueprintUnloaded.Remove(OnPackageUnloadedDelegateHandle);
 			}
 		}
 	}
diff --git a/Engine/Source/Developer/MovieSceneTools/Private/MovieSceneToolsModule.cpp b/Engine/Source/Developer/MovieSceneTools/Private/MovieSceneToolsModule.cpp
--- a/Engine/Source/Developer/MovieSceneTools/Private/MovieSceneToolsModule.cpp
+++ b/Engine/Source/Developer/MovieSceneTools/Private/MovieSceneToolsModule.cpp
@@ -37,13 +37,13 @@ class FMovieSceneToolsModule : public IMovieSceneTools
 
 		// Register with the sequencer module that we provide auto-key handlers.
 		ISequencerModule& SequencerModule = FModuleManager::Get().LoadModuleChecked<ISequencerModule>( "Sequencer" );
-		SequencerModule.RegisterTrackEditor( FOnCreateTrackEditor::CreateStatic( &FPropertyTrackEditor::CreateTrackEditor ) );
-		SequencerModule.RegisterTrackEditor( FOnCreateTrackEditor::CreateStatic( &F3DTransformTrackEditor::CreateTrackEditor ) );
-		SequencerModule.RegisterTrackEditor( FOnCreateTrackEditor::CreateStatic( &FDirectorTrackEditor::CreateTrackEditor ) );
-		SequencerModule.RegisterTrackEditor( FOnCreateTrackEditor::CreateStatic( &FSubMovieSceneTrackEditor::CreateTrackEditor ) );
-		SequencerModule.RegisterTrackEditor( FOnCreateTrackEditor::CreateStatic( &FAudioTrackEditor::CreateTrackEditor ) );
-		SequencerModule.RegisterTrackEditor( FOnCreateTrackEditor::CreateStatic( &FAnimationTrackEditor::CreateTrackEditor ) );
-		SequencerModule.RegisterTrackEditor( FOnCreateTrackEditor::CreateStatic( &FParticleTrackEditor::CreateTrackEditor ) );
+		PropertyTrackEditorCreateTrackEditorDelegateHandle      = SequencerModule.RegisterTrackEditor_Handle( FOnCreateTrackEditor::CreateStatic( &FPropertyTrackEditor::CreateTrackEditor ) );
+		TransformTrackEditorCreateTrackEditorDelegateHandle     = SequencerModule.RegisterTrackEditor_Handle( FOnCreateTrackEditor::CreateStatic( &F3DTransformTrackEditor::CreateTrackEditor ) );
+		DirectorTrackEditorCreateTrackEditorDelegateHandle      = SequencerModule.RegisterTrackEditor_Handle( FOnCreateTrackEditor::CreateStatic( &FDirectorTrackEditor::CreateTrackEditor ) );
+		SubMovieSceneTrackEditorCreateTrackEditorDelegateHandle = SequencerModule.RegisterTrackEditor_Handle( FOnCreateTrackEditor::CreateStatic( &FSubMovieSceneTrackEditor::CreateTrackEditor ) );
+		AudioTrackEditorCreateTrackEditorDelegateHandle         = SequencerModule.RegisterTrackEditor_Handle( FOnCreateTrackEditor::CreateStatic( &FAudioTrackEditor::CreateTrackEditor ) );
+		AnimationTrackEditorCreateTrackEditorDelegateHandle     = SequencerModule.RegisterTrackEditor_Handle( FOnCreateTrackEditor::CreateStatic( &FAnimationTrackEditor::CreateTrackEditor ) );
+		ParticleTrackEditorCreateTrackEditorDelegateHandle      = SequencerModule.RegisterTrackEditor_Handle( FOnCreateTrackEditor::CreateStatic( &FParticleTrackEditor::CreateTrackEditor ) );
 	}
 
 
@@ -62,13 +62,13 @@ class FMovieSceneToolsModule : public IMovieSceneTools
 		{
 			// Unregister auto key handlers
 			ISequencerModule& SequencerModule = FModuleManager::Get().GetModuleChecked<ISequencerModule>( "Sequencer" );
-			SequencerModule.UnRegisterTrackEditor( FOnCreateTrackEditor::CreateStatic( &FPropertyTrackEditor::CreateTrackEditor ) );
-			SequencerModule.UnRegisterTrackEditor( FOnCreateTrackEditor::CreateStatic( &F3DTransformTrackEditor::CreateTrackEditor ) );
-			SequencerModule.UnRegisterTrackEditor( FOnCreateTrackEditor::CreateStatic( &FDirectorTrackEditor::CreateTrackEditor ) );
-			SequencerModule.UnRegisterTrackEditor( FOnCreateTrackEditor::CreateStatic( &FSubMovieSceneTrackEditor::CreateTrackEditor ) );
-			SequencerModule.UnRegisterTrackEditor( FOnCreateTrackEditor::CreateStatic( &FAudioTrackEditor::CreateTrackEditor ) );
-			SequencerModule.UnRegisterTrackEditor( FOnCreateTrackEditor::CreateStatic( &FAnimationTrackEditor::CreateTrackEditor ) );
-			SequencerModule.UnRegisterTrackEditor( FOnCreateTrackEditor::CreateStatic( &FParticleTrackEditor::CreateTrackEditor ) );
+			SequencerModule.UnRegisterTrackEditor_Handle( PropertyTrackEditorCreateTrackEditorDelegateHandle );
+			SequencerModule.UnRegisterTrackEditor_Handle( TransformTrackEditorCreateTrackEditorDelegateHandle );
+			SequencerModule.UnRegisterTrackEditor_Handle( DirectorTrackEditorCreateTrackEditorDelegateHandle );
+			SequencerModule.UnRegisterTrackEditor_Handle( SubMovieSceneTrackEditorCreateTrackEditorDelegateHandle );
+			SequencerModule.UnRegisterTrackEditor_Handle( AudioTrackEditorCreateTrackEditorDelegateHandle );
+			SequencerModule.UnRegisterTrackEditor_Handle( AnimationTrackEditorCreateTrackEditorDelegateHandle );
+			SequencerModule.UnRegisterTrackEditor_Handle( ParticleTrackEditorCreateTrackEditorDelegateHandle );
 		}
 		
 	}
@@ -78,6 +78,15 @@ private:
 
 	/** Asset type actions for MovieScene assets.  Cached here so that we can unregister it during shutdown. */
 	TSharedPtr< FAssetTypeActions_MovieScene > MovieSceneAssetTypeActions;
+
+	/** Registered delegate handles */
+	FDelegateHandle PropertyTrackEditorCreateTrackEditorDelegateHandle;
+	FDelegateHandle TransformTrackEditorCreateTrackEditorDelegateHandle;
+	FDelegateHandle DirectorTrackEditorCreateTrackEditorDelegateHandle;
+	FDelegateHandle SubMovieSceneTrackEditorCreateTrackEditorDelegateHandle;
+	FDelegateHandle AudioTrackEditorCreateTrackEditorDelegateHandle;
+	FDelegateHandle AnimationTrackEditorCreateTrackEditorDelegateHandle;
+	FDelegateHandle ParticleTrackEditorCreateTrackEditorDelegateHandle;
 };
 
 
diff --git a/Engine/Source/Developer/Profiler/Private/ProfilerManager.cpp b/Engine/Source/Developer/Profiler/Private/ProfilerManager.cpp
--- a/Engine/Source/Developer/Profiler/Private/ProfilerManager.cpp
+++ b/Engine/Source/Developer/Profiler/Private/ProfilerManager.cpp
@@ -100,7 +100,7 @@ void FProfilerManager::PostConstructor()
 {
 	// Register tick functions.
 	OnTick = FTickerDelegate::CreateSP( this, &FProfilerManager::Tick );
-	FTicker::GetCoreTicker().AddTicker( OnTick );
+	OnTickHandle = FTicker::GetCoreTicker().AddTicker( OnTick );
 
 	// Create profiler client.
 	ProfilerClient = FModuleManager::GetModuleChecked<IProfilerClientModule>("ProfilerClient").CreateProfilerClient();
@@ -143,7 +143,7 @@ FProfilerManager::~FProfilerManager()
 	FProfilerCommands::Unregister();
 
 	// Unregister tick function.
-	FTicker::GetCoreTicker().RemoveTicker( OnTick );
+	FTicker::GetCoreTicker().RemoveTicker( OnTickHandle );
 
 	// Remove ourselves from the profiler client.
 	if( ProfilerClient.IsValid() )
diff --git a/Engine/Source/Developer/Profiler/Private/ProfilerManager.h b/Engine/Source/Developer/Profiler/Private/ProfilerManager.h
--- a/Engine/Source/Developer/Profiler/Private/ProfilerManager.h
+++ b/Engine/Source/Developer/Profiler/Private/ProfilerManager.h
@@ -588,6 +588,9 @@ protected:
 	/** The delegate to be invoked when this profiler manager ticks. */
 	FTickerDelegate OnTick;
 
+	/** Handle to the registered OnTick. */
+	FDelegateHandle OnTickHandle;
+
 	/** A weak pointer to the profiler window. */
 	TWeakPtr<class SProfilerWindow> ProfilerWindow;
 
diff --git a/Engine/Source/Developer/Profiler/Private/ProfilerSession.cpp b/Engine/Source/Developer/Profiler/Private/ProfilerSession.cpp
--- a/Engine/Source/Developer/Profiler/Private/ProfilerSession.cpp
+++ b/Engine/Source/Developer/Profiler/Private/ProfilerSession.cpp
@@ -67,7 +67,7 @@ FProfilerSession::FProfilerSession( const ISessionInstanceInfoPtr InSessionInsta
 {
 	// Randomize creation time to test loading profiler captures with different creation time and different amount of data.
 	CreationTime = FDateTime::Now() += FTimespan( 0, 0, FMath::RandRange( 2, 8 ) );
-	FTicker::GetCoreTicker().AddTicker( OnTick );
+	OnTickHandle = FTicker::GetCoreTicker().AddTicker( OnTick );
 }
 
 FProfilerSession::FProfilerSession( const FString& InDataFilepath )
@@ -94,12 +94,12 @@ FProfilerSession::FProfilerSession( const FString& InDataFilepath )
 {
 	// Randomize creation time to test loading profiler captures with different creation time and different amount of data.
 	CreationTime = FDateTime::Now() += FTimespan( 0, 0, FMath::RandRange( 2, 8 ) );
-	FTicker::GetCoreTicker().AddTicker( OnTick );
+	OnTickHandle = FTicker::GetCoreTicker().AddTicker( OnTick );
 }
 
 FProfilerSession::~FProfilerSession()
 {
-	FTicker::GetCoreTicker().RemoveTicker( OnTick );	
+	FTicker::GetCoreTicker().RemoveTicker( OnTickHandle );	
 }
 
 FGraphDataSourceRefConst FProfilerSession::CreateGraphDataSource( const uint32 InStatID )
@@ -457,7 +457,7 @@ FRawProfilerSession::FRawProfilerSession( const FString& InRawStatsFileFileath )
 
 FRawProfilerSession::~FRawProfilerSession()
 {
-	FTicker::GetCoreTicker().RemoveTicker( OnTick );
+	FTicker::GetCoreTicker().RemoveTicker( OnTickHandle );
 }
 
 bool FRawProfilerSession::HandleTicker( float DeltaTime )
@@ -726,7 +726,7 @@ void FRawProfilerSession::PrepareLoading()
 	const int64 AllocatedSize = ProfilerStream.GetAllocatedSize();
 
 	// We have the whole metadata and basic information about the raw stats file, start ticking the profiler session.
-	//FTicker::GetCoreTicker().AddTicker( OnTick, 0.25f );
+	//OnTickHandle = FTicker::GetCoreTicker().AddTicker( OnTick, 0.25f );
 
 #if	0
 	if( SessionType == EProfilerSessionTypes::OfflineRaw )
diff --git a/Engine/Source/Developer/Profiler/Private/ProfilerSession.h b/Engine/Source/Developer/Profiler/Private/ProfilerSession.h
--- a/Engine/Source/Developer/Profiler/Private/ProfilerSession.h
+++ b/Engine/Source/Developer/Profiler/Private/ProfilerSession.h
@@ -1214,6 +1214,9 @@ protected:
 	/** The delegate to be invoked when this profiler session instance ticks. */
 	FTickerDelegate OnTick;
 
+	/** Handle to the registered OnTick. */
+	FDelegateHandle OnTickHandle;
+
 	/** The data provider which holds all the collected profiler samples. */
 	IDataProviderRef DataProvider;
 
diff --git a/Engine/Source/Developer/ProfilerClient/Private/ProfilerClientManager.cpp b/Engine/Source/Developer/ProfilerClient/Private/ProfilerClientManager.cpp
--- a/Engine/Source/Developer/ProfilerClient/Private/ProfilerClientManager.cpp
+++ b/Engine/Source/Developer/ProfilerClient/Private/ProfilerClientManager.cpp
@@ -54,7 +54,7 @@ FProfilerClientManager::FProfilerClientManager( const IMessageBusRef& InMessageB
 	LoadTask = nullptr;
 #endif
 	LoadConnection = nullptr;
-	FTicker::GetCoreTicker().AddTicker(MessageDelegate, 0.05f);
+	MessageDelegateHandle = FTicker::GetCoreTicker().AddTicker(MessageDelegate, 0.05f);
 #endif
 }
 
@@ -74,8 +74,8 @@ FProfilerClientManager::~FProfilerClientManager()
 		UE_LOG(LogProfile, Log, TEXT( "File service-client transfer aborted: %s" ), *It.Key() );
 	}
 
-	FTicker::GetCoreTicker().RemoveTicker(MessageDelegate);
-	FTicker::GetCoreTicker().RemoveTicker(TickDelegate);
+	FTicker::GetCoreTicker().RemoveTicker(MessageDelegateHandle);
+	FTicker::GetCoreTicker().RemoveTicker(TickDelegateHandle);
 
 	Unsubscribe();
 
@@ -130,7 +130,7 @@ void FProfilerClientManager::Track( const FGuid& Instance )
 		MessageEndpoint->Publish(new FProfilerServiceSubscribe(ActiveSessionId, Instance), EMessageScope::Network);
 
 		RetryTime = 5.f;
-		FTicker::GetCoreTicker().AddTicker(TickDelegate, RetryTime);
+		TickDelegateHandle = FTicker::GetCoreTicker().AddTicker(TickDelegate, RetryTime);
 	}
 #endif
 }
@@ -303,7 +303,7 @@ void FProfilerClientManager::LoadCapture( const FString& DataFilepath, const FGu
 #endif
 
 	RetryTime = 0.05f;
-	FTicker::GetCoreTicker().AddTicker(TickDelegate, RetryTime);
+	TickDelegateHandle = FTicker::GetCoreTicker().AddTicker(TickDelegate, RetryTime);
 	ProfilerLoadStartedDelegate.Broadcast(ProfileId);
 #endif
 }
diff --git a/Engine/Source/Developer/ProfilerClient/Private/ProfilerClientManager.h b/Engine/Source/Developer/ProfilerClient/Private/ProfilerClientManager.h
--- a/Engine/Source/Developer/ProfilerClient/Private/ProfilerClientManager.h
+++ b/Engine/Source/Developer/ProfilerClient/Private/ProfilerClientManager.h
@@ -330,6 +330,9 @@ private:
 	/** Holds a delegate to be invoked when the widget ticks. */
 	FTickerDelegate TickDelegate;
 
+	/** Handle to the registered TickDelegate. */
+	FDelegateHandle TickDelegateHandle;
+
 	/** Amount of time between connection retries */
 	float RetryTime;
 
@@ -339,6 +342,9 @@ private:
 	/** Holds a delegate to be invoked when the widget ticks. */
 	FTickerDelegate MessageDelegate;
 
+	/** Handle to the registered MessageDelegate. */
+	FDelegateHandle MessageDelegateHandle;
+
 #if PROFILER_THREADED_LOAD
 	/** Loads a file asynchronously */
 	bool AsyncLoad();
diff --git a/Engine/Source/Developer/ProfilerService/Private/ProfilerServiceManager.cpp b/Engine/Source/Developer/ProfilerService/Private/ProfilerServiceManager.cpp
--- a/Engine/Source/Developer/ProfilerService/Private/ProfilerServiceManager.cpp
+++ b/Engine/Source/Developer/ProfilerService/Private/ProfilerServiceManager.cpp
@@ -531,7 +531,7 @@ void FProfilerServiceManager::SetPreviewState( const FMessageAddress& ClientAddr
 				// enable stat capture
 				if (PreviewClients.Num() == 0)
 				{
-					Stats.NewFrameDelegate.AddRaw( this, &FProfilerServiceManager::HandleNewFrame );
+					HandleNewFrameDelegateHandle = Stats.NewFrameDelegate.AddRaw( this, &FProfilerServiceManager::HandleNewFrame );
 					StatsMasterEnableAdd();
 				}
 				PreviewClients.Add(ClientAddress);
@@ -552,7 +552,7 @@ void FProfilerServiceManager::SetPreviewState( const FMessageAddress& ClientAddr
 				if (PreviewClients.Num() == 0)
 				{
 					// disable stat capture
-					Stats.NewFrameDelegate.RemoveRaw( this, &FProfilerServiceManager::HandleNewFrame );
+					Stats.NewFrameDelegate.Remove( HandleNewFrameDelegateHandle );
 					StatsMasterEnableAdd();
 				}
 			}
@@ -592,7 +592,7 @@ bool FProfilerServiceManager::HandlePing( float DeltaTime )
 	if (PreviewClients.Num() == 0)
 	{
 		// disable stat capture
-		FStatsThreadState::GetLocalState().NewFrameDelegate.RemoveRaw(this, &FProfilerServiceManager::HandleNewFrame);
+		FStatsThreadState::GetLocalState().NewFrameDelegate.Remove(HandleNewFrameDelegateHandle);
 		StatsMasterEnableAdd();
 	}
 
@@ -705,7 +705,7 @@ void FProfilerServiceManager::HandleServiceSubscribeMessage( const FProfilerServ
 		// initiate the ping callback
 		if (ClientData.Num() == 1)
 		{
-			FTicker::GetCoreTicker().AddTicker(PingDelegate, 15.0f);
+			PingDelegateHandle = FTicker::GetCoreTicker().AddTicker(PingDelegate, 15.0f);
 		}
 	}
 #endif
@@ -732,7 +732,7 @@ void FProfilerServiceManager::HandleServiceUnsubscribeMessage( const FProfilerSe
 		// stop the ping messages if we have no clients
 		if (ClientData.Num() == 0)
 		{
-			FTicker::GetCoreTicker().RemoveTicker(PingDelegate);
+			FTicker::GetCoreTicker().RemoveTicker(PingDelegateHandle);
 		}
 	}
 }
diff --git a/Engine/Source/Developer/ProfilerService/Private/ProfilerServiceManager.h b/Engine/Source/Developer/ProfilerService/Private/ProfilerServiceManager.h
--- a/Engine/Source/Developer/ProfilerService/Private/ProfilerServiceManager.h
+++ b/Engine/Source/Developer/ProfilerService/Private/ProfilerServiceManager.h
@@ -174,4 +174,10 @@ private:
 
 	/** Holds a delegate to be invoked for client pings */
 	FTickerDelegate PingDelegate;
+
+	/** Handle to the registered PingDelegate */
+	FDelegateHandle PingDelegateHandle;
+
+	/** Handle to the registered HandleNewFrame delegate */
+	FDelegateHandle HandleNewFrameDelegateHandle;
 };
diff --git a/Engine/Source/Developer/SessionServices/Private/SessionManager.cpp b/Engine/Source/Developer/SessionServices/Private/SessionManager.cpp
--- a/Engine/Source/Developer/SessionServices/Private/SessionManager.cpp
+++ b/Engine/Source/Developer/SessionServices/Private/SessionManager.cpp
@@ -35,7 +35,7 @@ FSessionManager::FSessionManager( const IMessageBusRef& InMessageBus )
 
 	// initialize ticker
 	TickDelegate = FTickerDelegate::CreateRaw(this, &FSessionManager::HandleTicker);
-	FTicker::GetCoreTicker().AddTicker(TickDelegate, 1.f);
+	TickDelegateHandle = FTicker::GetCoreTicker().AddTicker(TickDelegate, 1.f);
 
 	SendPing();
 }
@@ -43,7 +43,7 @@ FSessionManager::FSessionManager( const IMessageBusRef& InMessageBus )
 
 FSessionManager::~FSessionManager()
 {
-	FTicker::GetCoreTicker().RemoveTicker(TickDelegate);
+	FTicker::GetCoreTicker().RemoveTicker(TickDelegateHandle);
 }
 
 
diff --git a/Engine/Source/Developer/SessionServices/Private/SessionManager.h b/Engine/Source/Developer/SessionServices/Private/SessionManager.h
--- a/Engine/Source/Developer/SessionServices/Private/SessionManager.h
+++ b/Engine/Source/Developer/SessionServices/Private/SessionManager.h
@@ -151,4 +151,7 @@ private:
 
 	/** Holds a delegate to be invoked when the widget ticks. */
 	FTickerDelegate TickDelegate;
+
+	/** Holds a delegate to be invoked when the widget ticks. */
+	FDelegateHandle TickDelegateHandle;
 };
diff --git a/Engine/Source/Developer/SourceControl/Private/DefaultSourceControlProvider.cpp b/Engine/Source/Developer/SourceControl/Private/DefaultSourceControlProvider.cpp
--- a/Engine/Source/Developer/SourceControl/Private/DefaultSourceControlProvider.cpp
+++ b/Engine/Source/Developer/SourceControl/Private/DefaultSourceControlProvider.cpp
@@ -55,6 +55,16 @@ void FDefaultSourceControlProvider::UnregisterSourceControlStateChanged( const F
 
 }
 
+FDelegateHandle FDefaultSourceControlProvider::RegisterSourceControlStateChanged_Handle( const FSourceControlStateChanged::FDelegate& SourceControlStateChanged )
+{
+	return FDelegateHandle();
+}
+
+void FDefaultSourceControlProvider::UnregisterSourceControlStateChanged_Handle( FDelegateHandle Handle )
+{
+
+}
+
 ECommandResult::Type FDefaultSourceControlProvider::Execute( const TSharedRef<ISourceControlOperation, ESPMode::ThreadSafe>& InOperation, const TArray<FString>& InFiles, EConcurrency::Type InConcurrency, const FSourceControlOperationComplete& InOperationCompleteDelegate )
 {
 	return ECommandResult::Failed;
diff --git a/Engine/Source/Developer/SourceControl/Private/DefaultSourceControlProvider.h b/Engine/Source/Developer/SourceControl/Private/DefaultSourceControlProvider.h
--- a/Engine/Source/Developer/SourceControl/Private/DefaultSourceControlProvider.h
+++ b/Engine/Source/Developer/SourceControl/Private/DefaultSourceControlProvider.h
@@ -4,6 +4,13 @@
 
 #include "ISourceControlProvider.h"
 
+// We disable the deprecation warnings here because otherwise it'll complain about us
+// implementing RegisterSourceControlStateChanged and UnregisterSourceControlStateChanged.  We know
+// that, but we only want it to complain if *others* implement or call these functions.
+//
+// These macros should be removed when those functions are removed.
+PRAGMA_DISABLE_DEPRECATION_WARNINGS
+
 /**
  * Default source control provider implementation - "None"
  */
@@ -20,6 +27,8 @@ public:
 	virtual ECommandResult::Type GetState( const TArray<FString>& InFiles, TArray< TSharedRef<ISourceControlState, ESPMode::ThreadSafe> >& OutState, EStateCacheUsage::Type InStateCacheUsage ) override;
 	virtual void RegisterSourceControlStateChanged( const FSourceControlStateChanged::FDelegate& SourceControlStateChanged ) override;
 	virtual void UnregisterSourceControlStateChanged( const FSourceControlStateChanged::FDelegate& SourceControlStateChanged ) override;
+	virtual FDelegateHandle RegisterSourceControlStateChanged_Handle( const FSourceControlStateChanged::FDelegate& SourceControlStateChanged ) override;
+	virtual void UnregisterSourceControlStateChanged_Handle( FDelegateHandle Handle ) override;
 	virtual ECommandResult::Type Execute( const TSharedRef<ISourceControlOperation, ESPMode::ThreadSafe>& InOperation, const TArray<FString>& InFiles, EConcurrency::Type InConcurrency = EConcurrency::Synchronous, const FSourceControlOperationComplete& InOperationCompleteDelegate = FSourceControlOperationComplete() ) override;
 	virtual bool CanCancelOperation( const TSharedRef<ISourceControlOperation, ESPMode::ThreadSafe>& InOperation ) const override;
 	virtual void CancelOperation( const TSharedRef<ISourceControlOperation, ESPMode::ThreadSafe>& InOperation ) override;
@@ -30,3 +39,5 @@ public:
 	virtual TSharedRef<class SWidget> MakeSettingsWidget() const override;
 #endif // SOURCE_CONTROL_WITH_SLATE
 };
+
+PRAGMA_ENABLE_DEPRECATION_WARNINGS
diff --git a/Engine/Source/Developer/SourceControl/Public/ISourceControlProvider.h b/Engine/Source/Developer/SourceControl/Public/ISourceControlProvider.h
--- a/Engine/Source/Developer/SourceControl/Public/ISourceControlProvider.h
+++ b/Engine/Source/Developer/SourceControl/Public/ISourceControlProvider.h
@@ -162,13 +162,25 @@ public:
 	/**
 	 * Register a delegate to be called when source control state(s) change
 	 */
+	DELEGATE_DEPRECATED("This function is deprecated - please replace any usage with RegisterSourceControlStateChanged_Handle.")
 	virtual void RegisterSourceControlStateChanged( const FSourceControlStateChanged::FDelegate& SourceControlStateChanged ) = 0;
 
 	/**
 	 * Unregister a delegate to be called when source control state(s) change
 	 */
+	DELEGATE_DEPRECATED("This function is deprecated - please replace any usage with UnregisterSourceControlStateChanged_Handle, passing the result of RegisterSourceControlStateChanged_Handle.")
 	virtual void UnregisterSourceControlStateChanged( const FSourceControlStateChanged::FDelegate& SourceControlStateChanged ) = 0;
 
+	/**
+	 * Register a delegate to be called when source control state(s) change
+	 */
+	virtual FDelegateHandle RegisterSourceControlStateChanged_Handle( const FSourceControlStateChanged::FDelegate& SourceControlStateChanged ) = 0;
+
+	/**
+	 * Unregister a delegate to be called when source control state(s) change
+	 */
+	virtual void UnregisterSourceControlStateChanged_Handle( FDelegateHandle Handle ) = 0;
+
 	/**
 	 * Attempt to execute an operation on the passed-in files (if any are required).
 	 * @param	InOperation						The operation to perform.
diff --git a/Engine/Source/Developer/TargetDeviceServices/Private/TargetDeviceProxyManager.cpp b/Engine/Source/Developer/TargetDeviceServices/Private/TargetDeviceProxyManager.cpp
--- a/Engine/Source/Developer/TargetDeviceServices/Private/TargetDeviceProxyManager.cpp
+++ b/Engine/Source/Developer/TargetDeviceServices/Private/TargetDeviceProxyManager.cpp
@@ -14,7 +14,7 @@ FTargetDeviceProxyManager::FTargetDeviceProxyManager()
 	if (MessageEndpoint.IsValid())
 	{
 		TickDelegate = FTickerDelegate::CreateRaw(this, &FTargetDeviceProxyManager::HandleTicker);
-		FTicker::GetCoreTicker().AddTicker(TickDelegate, TARGET_DEVICE_SERVICES_PING_INTERVAL);
+		TickDelegateHandle = FTicker::GetCoreTicker().AddTicker(TickDelegate, TARGET_DEVICE_SERVICES_PING_INTERVAL);
 
 		SendPing();
 	}
@@ -23,7 +23,7 @@ FTargetDeviceProxyManager::FTargetDeviceProxyManager()
 
 FTargetDeviceProxyManager::~FTargetDeviceProxyManager()
 {
-	FTicker::GetCoreTicker().RemoveTicker(TickDelegate);
+	FTicker::GetCoreTicker().RemoveTicker(TickDelegateHandle);
 }
 
 
diff --git a/Engine/Source/Developer/TargetDeviceServices/Private/TargetDeviceProxyManager.h b/Engine/Source/Developer/TargetDeviceServices/Private/TargetDeviceProxyManager.h
--- a/Engine/Source/Developer/TargetDeviceServices/Private/TargetDeviceProxyManager.h
+++ b/Engine/Source/Developer/TargetDeviceServices/Private/TargetDeviceProxyManager.h
@@ -74,4 +74,7 @@ private:
 
 	/** Holds a delegate to be invoked when the widget ticks. */
 	FTickerDelegate TickDelegate;
+
+	/** Handle to the registered TickDelegate. */
+	FDelegateHandle TickDelegateHandle;
 };
diff --git a/Engine/Source/Developer/iOS/IOSTargetPlatform/Private/IOSTargetPlatform.cpp b/Engine/Source/Developer/iOS/IOSTargetPlatform/Private/IOSTargetPlatform.cpp
--- a/Engine/Source/Developer/iOS/IOSTargetPlatform/Private/IOSTargetPlatform.cpp
+++ b/Engine/Source/Developer/iOS/IOSTargetPlatform/Private/IOSTargetPlatform.cpp
@@ -20,7 +20,7 @@ FIOSTargetPlatform::FIOSTargetPlatform()
 
 	// Initialize Ticker for device discovery
 	TickDelegate = FTickerDelegate::CreateRaw(this, &FIOSTargetPlatform::HandleTicker);
-	FTicker::GetCoreTicker().AddTicker(TickDelegate, 10.0f);
+	TickDelegateHandle = FTicker::GetCoreTicker().AddTicker(TickDelegate, 10.0f);
 	
 	// initialize the connected device detector
 	DeviceHelper.OnDeviceConnected().AddRaw(this, &FIOSTargetPlatform::HandleDeviceConnected);
@@ -31,7 +31,7 @@ FIOSTargetPlatform::FIOSTargetPlatform()
 
 FIOSTargetPlatform::~FIOSTargetPlatform()
 {
-	FTicker::GetCoreTicker().RemoveTicker(TickDelegate);
+	FTicker::GetCoreTicker().RemoveTicker(TickDelegateHandle);
 }
 
 
diff --git a/Engine/Source/Editor/AddContentDialog/Private/ContentSourceProviders/FeaturePack/FeaturePackContentSourceProvider.cpp b/Engine/Source/Editor/AddContentDialog/Private/ContentSourceProviders/FeaturePack/FeaturePackContentSourceProvider.cpp
--- a/Engine/Source/Editor/AddContentDialog/Private/ContentSourceProviders/FeaturePack/FeaturePackContentSourceProvider.cpp
+++ b/Engine/Source/Editor/AddContentDialog/Private/ContentSourceProviders/FeaturePack/FeaturePackContentSourceProvider.cpp
@@ -53,7 +53,7 @@ void FFeaturePackContentSourceProvider::StartUpDirectoryWatcher()
 		// If the path doesn't exist on disk, make it so the watcher will work.
 		IFileManager::Get().MakeDirectory(*FeaturePackPath);
 		DirectoryChangedDelegate = IDirectoryWatcher::FDirectoryChanged::CreateRaw(this, &FFeaturePackContentSourceProvider::OnFeaturePackDirectoryChanged);
-		DirectoryWatcher->RegisterDirectoryChangedCallback(FeaturePackPath, DirectoryChangedDelegate);
+		DirectoryWatcher->RegisterDirectoryChangedCallback_Handle(FeaturePackPath, DirectoryChangedDelegate, DirectoryChangedDelegateHandle);
 	}
 }
 
@@ -63,7 +63,7 @@ void FFeaturePackContentSourceProvider::ShutDownDirectoryWatcher()
 	IDirectoryWatcher* DirectoryWatcher = DirectoryWatcherModule.Get();
 	if ( DirectoryWatcher )
 	{
-		DirectoryWatcher->UnregisterDirectoryChangedCallback(FeaturePackPath, DirectoryChangedDelegate);
+		DirectoryWatcher->UnregisterDirectoryChangedCallback_Handle(FeaturePackPath, DirectoryChangedDelegateHandle);
 	}
 }
 
diff --git a/Engine/Source/Editor/AddContentDialog/Private/ContentSourceProviders/FeaturePack/FeaturePackContentSourceProvider.h b/Engine/Source/Editor/AddContentDialog/Private/ContentSourceProviders/FeaturePack/FeaturePackContentSourceProvider.h
--- a/Engine/Source/Editor/AddContentDialog/Private/ContentSourceProviders/FeaturePack/FeaturePackContentSourceProvider.h
+++ b/Engine/Source/Editor/AddContentDialog/Private/ContentSourceProviders/FeaturePack/FeaturePackContentSourceProvider.h
@@ -39,4 +39,6 @@ private:
 
 	/** An array of the available content sources. */
 	TArray<TSharedRef<IContentSource>> ContentSources;
+
+	FDelegateHandle DirectoryChangedDelegateHandle;
 };
\ No newline at end of file
