#FriendsAndChat - Merge from Fornite Dev

diff --git a/Engine/Source/Developer/FriendsAndChat/Private/FriendsAndChatManager.cpp b/Engine/Source/Developer/FriendsAndChat/Private/FriendsAndChatManager.cpp
--- a/Engine/Source/Developer/FriendsAndChat/Private/FriendsAndChatManager.cpp
+++ b/Engine/Source/Developer/FriendsAndChat/Private/FriendsAndChatManager.cpp
@@ -1147,16 +1147,25 @@ void FFriendsAndChatManager::SendInviteAcceptedNotification(const TSharedPtr< IF
 {
 	if(OnFriendsActionNotification().IsBound())
 	{
-		FFormatNamedArguments Args;
-		Args.Add(TEXT("Username"), FText::FromString(Friend->GetName()));
-		const FText FriendRequestMessage = FText::Format(LOCTEXT("FFriendsAndChatManager_Accepted", "{Username} accepted your request"), Args);
-
+		const FText FriendRequestMessage = GetInviteNotificationText(Friend);
 		TSharedPtr< FFriendsAndChatMessage > NotificationMessage = MakeShareable(new FFriendsAndChatMessage(FriendRequestMessage.ToString()));
 		NotificationMessage->SetMessageType(EFriendsRequestType::FriendAccepted);
 		OnFriendsActionNotification().Broadcast(NotificationMessage.ToSharedRef());
 	}
 }
 
+const FText FFriendsAndChatManager::GetInviteNotificationText(TSharedPtr< IFriendItem > Friend) const
+{
+	FFormatNamedArguments Args;
+	Args.Add(TEXT("Username"), FText::FromString(Friend->GetName()));
+
+	if(Friend->IsPendingAccepted())
+	{
+		return FText::Format(LOCTEXT("FriendAddedToast", "{Username} added as a friend"), Args);
+	}
+	return FText::Format(LOCTEXT("FriendAcceptedToast", "{Username} accepted your request"), Args);
+}
+
 void FFriendsAndChatManager::OnQueryUserIdMappingComplete(bool bWasSuccessful, const FUniqueNetId& RequestingUserId, const FString& DisplayName, const FUniqueNetId& IdentifiedUserId, const FString& Error)
 {
 	check( OnlineSubMcp != nullptr && OnlineSubMcp->GetMcpPersonaService().IsValid() );
@@ -1218,7 +1227,7 @@ void FFriendsAndChatManager::OnQueryUserIdMappingComplete(bool bWasSuccessful, c
 			{
 				FOnSendInviteComplete Delegate = FOnSendInviteComplete::CreateSP(this, &FFriendsAndChatManager::OnSendInviteComplete);
 				FriendsInterface->SendInvite(0, PendingOutgoingFriendRequests[Index].Get(), EFriendsLists::ToString( EFriendsLists::Default ), Delegate);
-				AddFriendsToast(LOCTEXT("FFriendsAndChatManager_FriendRequestSent", "Request Sent"));
+				AddFriendsToast(LOCTEXT("FriendRequestSentToast", "Request Sent"));
 			}
 		}
 		else
@@ -1414,7 +1423,7 @@ void FFriendsAndChatManager::SendGameInvite(const FUniqueNetId& ToUser)
 		if (UserId.IsValid())
 		{
 			OnlineSubMcp->GetSessionInterface()->SendSessionInviteToFriend(*UserId, GameSessionName, ToUser);
-			AddFriendsToast(LOCTEXT("FFriendsAndChatManager_InviteToGameSent", "Invite Sent"));
+			AddFriendsToast(LOCTEXT("InviteToGameSentToast", "Invite Sent"));
 			Analytics.RecordGameInvite(ToUser, TEXT("Social.GameInvite.Send"));
 		}
 	}
@@ -1435,8 +1444,8 @@ void FFriendsAndChatManager::OnInviteAccepted(const FUniqueNetId& UserId, const
 	TSharedPtr< IFriendItem > Friend = FindUser(FriendId);
 	if(Friend.IsValid())
 	{
-		Friend->SetPendingAccept();
 		SendInviteAcceptedNotification(Friend);
+		Friend->SetPendingAccept();
 	}
 	RefreshList();
 	RequestListRefresh();
diff --git a/Engine/Source/Developer/FriendsAndChat/Private/FriendsAndChatManager.h b/Engine/Source/Developer/FriendsAndChat/Private/FriendsAndChatManager.h
--- a/Engine/Source/Developer/FriendsAndChat/Private/FriendsAndChatManager.h
+++ b/Engine/Source/Developer/FriendsAndChat/Private/FriendsAndChatManager.h
@@ -417,6 +417,9 @@ private:
 	/** Send a friend invite accepted notification. */
 	void SendInviteAcceptedNotification(TSharedPtr< IFriendItem > Friend);
 
+	/** Get Invite Notification Text */
+	const FText GetInviteNotificationText(TSharedPtr< IFriendItem > Friend) const;
+
 	/** Called when singleton is released. */
 	void ShutdownManager();
 
diff --git a/Engine/Source/Developer/FriendsAndChat/Private/UI/Widgets/SChatWindow.h b/Engine/Source/Developer/FriendsAndChat/Private/UI/Widgets/SChatWindow.h
--- a/Engine/Source/Developer/FriendsAndChat/Private/UI/Widgets/SChatWindow.h
+++ b/Engine/Source/Developer/FriendsAndChat/Private/UI/Widgets/SChatWindow.h
@@ -6,7 +6,7 @@ class SChatWindow : public SUserWidget
 {
 public:
 	SLATE_USER_ARGS(SChatWindow)
-	: _MaxChatLength(200)
+	: _MaxChatLength(128)
 	{ }
 		SLATE_ARGUMENT( const FFriendsAndChatStyle*, FriendStyle )
 		SLATE_ARGUMENT(EPopupMethod, Method)
