[INTEGRATE] Change 2416860 by [REDACTED_PERSON][REDACTED_HANDLE] on 2015/01/23 10:23:13

	[INTEGRATE] Change 2416366 by [REDACTED_PERSON][REDACTED_HANDLE] on 2015/01/22 22:06:42

		#online,mcp,friends
		- removed direct mcp dependencies from public friends/chat module
		- user mapping via persona service exposed from existing onlineuser interface

diff --git a/Engine/Source/Developer/FriendsAndChat/Private/FriendsAndChatManager.cpp b/Engine/Source/Developer/FriendsAndChat/Private/FriendsAndChatManager.cpp
--- a/Engine/Source/Developer/FriendsAndChat/Private/FriendsAndChatManager.cpp
+++ b/Engine/Source/Developer/FriendsAndChat/Private/FriendsAndChatManager.cpp
@@ -21,7 +21,7 @@
  *****************************************************************************/
 
 FFriendsAndChatManager::FFriendsAndChatManager( )
-	: OnlineSubMcp(nullptr)
+	: OnlineSub(nullptr)
 	, MessageManager(FFriendsMessageManagerFactory::Create())
 	, ManagerState ( EFriendsAndManagerState::OffLine )
 	, bIsInited( false )
@@ -46,27 +46,27 @@ void FFriendsAndChatManager::Login()
 
 	bIsInited = false;
 
-	OnlineSubMcp = static_cast< FOnlineSubsystemMcp* >( IOnlineSubsystem::Get( TEXT( "MCP" ) ) );
+	OnlineSub = IOnlineSubsystem::Get(TEXT("MCP"));
 
-	if (OnlineSubMcp != nullptr &&
-		OnlineSubMcp->GetMcpAccountMappingService().IsValid() &&
-		OnlineSubMcp->GetIdentityInterface().IsValid())
+	if (OnlineSub != nullptr &&
+		OnlineSub->GetUserInterface().IsValid() &&
+		OnlineSub->GetIdentityInterface().IsValid())
 	{
-		OnlineIdentity = OnlineSubMcp->GetIdentityInterface();
+		OnlineIdentity = OnlineSub->GetIdentityInterface();
 
 		if(OnlineIdentity->GetUniquePlayerId(0).IsValid())
 		{
-			IOnlineUserPtr UserInterface = OnlineSubMcp->GetUserInterface();
+			IOnlineUserPtr UserInterface = OnlineSub->GetUserInterface();
 			check(UserInterface.IsValid());
 
-			FriendsInterface = OnlineSubMcp->GetFriendsInterface();
+			FriendsInterface = OnlineSub->GetFriendsInterface();
 			check( FriendsInterface.IsValid() )
 
 			// Create delegates for list refreshes
 			OnQueryRecentPlayersCompleteDelegate = FOnQueryRecentPlayersCompleteDelegate           ::CreateRaw(this, &FFriendsAndChatManager::OnQueryRecentPlayersComplete);
 			OnFriendsListChangedDelegate         = FOnFriendsChangeDelegate                        ::CreateSP (this, &FFriendsAndChatManager::OnFriendsListChanged);
 			OnDeleteFriendCompleteDelegate       = FOnDeleteFriendCompleteDelegate                 ::CreateSP (this, &FFriendsAndChatManager::OnDeleteFriendComplete);
-			OnQueryUserIdMappingCompleteDelegate = FOnQueryUserIdMappingCompleteDelegate           ::CreateSP (this, &FFriendsAndChatManager::OnQueryUserIdMappingComplete);
+			OnQueryUserIdMappingCompleteDelegate = IOnlineUser::FOnQueryUserMappingComplete::CreateSP( this, &FFriendsAndChatManager::OnQueryUserIdMappingComplete );
 			OnQueryUserInfoCompleteDelegate      = FOnQueryUserInfoCompleteDelegate                ::CreateSP (this, &FFriendsAndChatManager::OnQueryUserInfoComplete);
 			OnPresenceReceivedCompleteDelegate   = FOnPresenceReceivedDelegate                     ::CreateSP (this, &FFriendsAndChatManager::OnPresenceReceived);
 			OnPresenceUpdatedCompleteDelegate    = IOnlinePresence::FOnPresenceTaskCompleteDelegate::CreateSP (this, &FFriendsAndChatManager::OnPresenceUpdated);
@@ -87,12 +87,9 @@ void FFriendsAndChatManager::Login()
 
 			OnQueryUserInfoCompleteDelegateHandle = UserInterface->AddOnQueryUserInfoCompleteDelegate_Handle(0, OnQueryUserInfoCompleteDelegate);
 
-			OnPresenceReceivedCompleteDelegateHandle = OnlineSubMcp->GetPresenceInterface()->AddOnPresenceReceivedDelegate_Handle      (OnPresenceReceivedCompleteDelegate);
-			OnGameInviteReceivedDelegateHandle       = OnlineSubMcp->GetSessionInterface ()->AddOnSessionInviteReceivedDelegate_Handle (OnGameInviteReceivedDelegate);
-			OnDestroySessionCompleteDelegateHandle   = OnlineSubMcp->GetSessionInterface ()->AddOnDestroySessionCompleteDelegate_Handle(OnDestroySessionCompleteDelegate);
-
-			FOnlinePersonaMcpPtr OnlinePersonaMcp = OnlineSubMcp->GetMcpPersonaService();
-			OnQueryUserIdMappingCompleteDelegateHandle = OnlinePersonaMcp->AddOnQueryUserIdMappingCompleteDelegate_Handle(OnQueryUserIdMappingCompleteDelegate);
+			OnPresenceReceivedCompleteDelegateHandle = OnlineSub->GetPresenceInterface()->AddOnPresenceReceivedDelegate_Handle      (OnPresenceReceivedCompleteDelegate);
+			OnGameInviteReceivedDelegateHandle       = OnlineSub->GetSessionInterface ()->AddOnSessionInviteReceivedDelegate_Handle (OnGameInviteReceivedDelegate);
+			OnDestroySessionCompleteDelegateHandle   = OnlineSub->GetSessionInterface ()->AddOnDestroySessionCompleteDelegate_Handle(OnDestroySessionCompleteDelegate);
 
 			ManagerState = EFriendsAndManagerState::Idle;
 
@@ -125,34 +122,30 @@ void FFriendsAndChatManager::Login()
 
 void FFriendsAndChatManager::Logout()
 {
-	if (OnlineSubMcp != nullptr)
-	{
-		if (OnlineSubMcp->GetFriendsInterface().IsValid())
+	if (OnlineSub != nullptr)
 		{
-			OnlineSubMcp->GetFriendsInterface()->ClearOnQueryRecentPlayersCompleteDelegate_Handle(   OnQueryRecentPlayersCompleteDelegateHandle);
-			OnlineSubMcp->GetFriendsInterface()->ClearOnFriendsChangeDelegate_Handle             (0, OnFriendsListChangedDelegateHandle);
-			OnlineSubMcp->GetFriendsInterface()->ClearOnInviteReceivedDelegate_Handle            (   OnFriendInviteReceivedDelegateHandle);
-			OnlineSubMcp->GetFriendsInterface()->ClearOnFriendRemovedDelegate_Handle             (   OnFriendRemovedDelegateHandle);
-			OnlineSubMcp->GetFriendsInterface()->ClearOnInviteRejectedDelegate_Handle            (   OnFriendInviteRejectedHandle);
-			OnlineSubMcp->GetFriendsInterface()->ClearOnInviteAcceptedDelegate_Handle            (   OnFriendInviteAcceptedHandle);
-			OnlineSubMcp->GetFriendsInterface()->ClearOnDeleteFriendCompleteDelegate_Handle      (0, OnDeleteFriendCompleteDelegateHandle);
-		}
-		if (OnlineSubMcp->GetPresenceInterface().IsValid())
+		if (OnlineSub->GetFriendsInterface().IsValid())
 		{
-			OnlineSubMcp->GetPresenceInterface()->ClearOnPresenceReceivedDelegate_Handle(OnPresenceReceivedCompleteDelegateHandle);
-		}
-		if (OnlineSubMcp->GetSessionInterface().IsValid())
+			OnlineSub->GetFriendsInterface()->ClearOnQueryRecentPlayersCompleteDelegate_Handle(   OnQueryRecentPlayersCompleteDelegateHandle);
+			OnlineSub->GetFriendsInterface()->ClearOnFriendsChangeDelegate_Handle             (0, OnFriendsListChangedDelegateHandle);
+			OnlineSub->GetFriendsInterface()->ClearOnInviteReceivedDelegate_Handle            (   OnFriendInviteReceivedDelegateHandle);
+			OnlineSub->GetFriendsInterface()->ClearOnFriendRemovedDelegate_Handle             (   OnFriendRemovedDelegateHandle);
+			OnlineSub->GetFriendsInterface()->ClearOnInviteRejectedDelegate_Handle            (   OnFriendInviteRejectedHandle);
+			OnlineSub->GetFriendsInterface()->ClearOnInviteAcceptedDelegate_Handle            (   OnFriendInviteAcceptedHandle);
+			OnlineSub->GetFriendsInterface()->ClearOnDeleteFriendCompleteDelegate_Handle      (0, OnDeleteFriendCompleteDelegateHandle);
+		}
+		if (OnlineSub->GetPresenceInterface().IsValid())
 		{
-			OnlineSubMcp->GetSessionInterface()->ClearOnSessionInviteReceivedDelegate_Handle (OnGameInviteReceivedDelegateHandle);
-			OnlineSubMcp->GetSessionInterface()->ClearOnDestroySessionCompleteDelegate_Handle(OnDestroySessionCompleteDelegateHandle);
+			OnlineSub->GetPresenceInterface()->ClearOnPresenceReceivedDelegate_Handle(OnPresenceReceivedCompleteDelegateHandle);
 		}
-		if (OnlineSubMcp->GetUserInterface().IsValid())
+		if (OnlineSub->GetSessionInterface().IsValid())
 		{
-			OnlineSubMcp->GetUserInterface()->ClearOnQueryUserInfoCompleteDelegate_Handle(0, OnQueryUserInfoCompleteDelegateHandle);
+			OnlineSub->GetSessionInterface()->ClearOnSessionInviteReceivedDelegate_Handle (OnGameInviteReceivedDelegateHandle);
+			OnlineSub->GetSessionInterface()->ClearOnDestroySessionCompleteDelegate_Handle(OnDestroySessionCompleteDelegateHandle);
 		}
-		if (OnlineSubMcp->GetMcpPersonaService().IsValid())
+		if (OnlineSub->GetUserInterface().IsValid())
 		{
-			OnlineSubMcp->GetMcpPersonaService()->ClearOnQueryUserIdMappingCompleteDelegate_Handle(OnQueryUserIdMappingCompleteDelegateHandle);
+			OnlineSub->GetUserInterface()->ClearOnQueryUserInfoCompleteDelegate_Handle(0, OnQueryUserInfoCompleteDelegateHandle);
 		}
 	}
 
@@ -182,7 +175,7 @@ void FFriendsAndChatManager::Logout()
 
 	MessageManager->LogOut();
 
-	OnlineSubMcp = nullptr;
+	OnlineSub = nullptr;
 	if ( UpdateFriendsTickerDelegate.IsBound() )
 	{
 		FTicker::GetCoreTicker().RemoveTicker( UpdateFriendsTickerDelegateHandle );
@@ -462,10 +455,10 @@ void FFriendsAndChatManager::SetChatWindowContents()
 				+ SHorizontalBox::Slot()
 				.AutoWidth()
 				[
-					SNew(SFriendsStatus, StatusViewModel)
-					.FriendStyle(&Style)
-					.Method(EPopupMethod::UseCurrentWindow)
-				]
+				SNew(SFriendsStatus, StatusViewModel)
+				.FriendStyle(&Style)
+				.Method(EPopupMethod::UseCurrentWindow)
+			]
 				+ SHorizontalBox::Slot()
 				[
 					SNew(SSpacer)
@@ -482,20 +475,20 @@ void FFriendsAndChatManager::SetChatWindowContents()
 
 void FFriendsAndChatManager::SetUserIsOnline(EOnlinePresenceState::Type OnlineState)
 {
-	if ( OnlineSubMcp != nullptr )
+	if ( OnlineSub != nullptr )
 	{
 		TSharedPtr<FUniqueNetId> UserId = OnlineIdentity->GetUniquePlayerId(0);
 		if (UserId.IsValid())
 		{
 			TSharedPtr<FOnlineUserPresence> CurrentPresence;
-			OnlineSubMcp->GetPresenceInterface()->GetCachedPresence(*UserId, CurrentPresence);
+			OnlineSub->GetPresenceInterface()->GetCachedPresence(*UserId, CurrentPresence);
 			FOnlineUserPresenceStatus NewStatus;
 			if (CurrentPresence.IsValid())
 			{
 				NewStatus = CurrentPresence->Status;
 			}
 			NewStatus.State = OnlineState;
-			OnlineSubMcp->GetPresenceInterface()->SetPresence(*UserId, NewStatus, OnPresenceUpdatedCompleteDelegate);
+			OnlineSub->GetPresenceInterface()->SetPresence(*UserId, NewStatus, OnPresenceUpdatedCompleteDelegate);
 		}
 	}
 }
@@ -608,11 +601,11 @@ int32 FFriendsAndChatManager::GetFilteredGameInviteList(TArray< TSharedPtr< IFri
 FString FFriendsAndChatManager::GetGameSessionId() const
 {
 	FString Result;
-	if (OnlineSubMcp != nullptr &&
+	if (OnlineSub != nullptr &&
 		OnlineIdentity.IsValid() &&
-		OnlineSubMcp->GetSessionInterface().IsValid())
+		OnlineSub->GetSessionInterface().IsValid())
 	{
-		const FNamedOnlineSession* GameSession = OnlineSubMcp->GetSessionInterface()->GetNamedSession(GameSessionName);
+		const FNamedOnlineSession* GameSession = OnlineSub->GetSessionInterface()->GetNamedSession(GameSessionName);
 		if (GameSession != nullptr)
 		{
 			TSharedPtr<FOnlineSessionInfo> UserSessionInfo = GameSession->SessionInfo;
@@ -627,10 +620,10 @@ FString FFriendsAndChatManager::GetGameSessionId() const
 
 bool FFriendsAndChatManager::IsInGameSession() const
 {	
-	if (OnlineSubMcp != nullptr &&
+	if (OnlineSub != nullptr &&
 		OnlineIdentity.IsValid() &&
-		OnlineSubMcp->GetSessionInterface().IsValid() &&
-		OnlineSubMcp->GetSessionInterface()->GetNamedSession(GameSessionName) != nullptr)
+		OnlineSub->GetSessionInterface().IsValid() &&
+		OnlineSub->GetSessionInterface()->GetNamedSession(GameSessionName) != nullptr)
 	{
 		return true;
 	}
@@ -641,20 +634,20 @@ bool FFriendsAndChatManager::IsInJoinableGameSession() const
 {
 	bool bIsGameJoinable = false;
 
-	if (OnlineSubMcp != nullptr &&
-		OnlineSubMcp->GetIdentityInterface().IsValid() &&
-		OnlineSubMcp->GetPresenceInterface().IsValid())
+	if (OnlineSub != nullptr &&
+		OnlineSub->GetIdentityInterface().IsValid() &&
+		OnlineSub->GetPresenceInterface().IsValid())
 	{
-		TSharedPtr<FUniqueNetId> CurrentUserId = OnlineSubMcp->GetIdentityInterface()->GetUniquePlayerId(0);
+		TSharedPtr<FUniqueNetId> CurrentUserId = OnlineSub->GetIdentityInterface()->GetUniquePlayerId(0);
 		if (CurrentUserId.IsValid())
 		{
-			TSharedPtr<FOnlineUserPresence> Presence;
-			OnlineSubMcp->GetPresenceInterface()->GetCachedPresence(*CurrentUserId, Presence);
-			if (Presence.IsValid())
-			{
-				bIsGameJoinable = Presence->bIsJoinable;
-			}
-		}
+		    TSharedPtr<FOnlineUserPresence> Presence;
+		    OnlineSub->GetPresenceInterface()->GetCachedPresence(*CurrentUserId, Presence);
+		    if (Presence.IsValid())
+		    {
+			    bIsGameJoinable = Presence->bIsJoinable;
+		    }
+	    }
 	}
 
 	return bIsGameJoinable && IsInGameSession();
@@ -681,13 +674,13 @@ const bool FFriendsAndChatManager::IsInLauncher() const
 
 EOnlinePresenceState::Type FFriendsAndChatManager::GetUserIsOnline()
 {
-	if (OnlineSubMcp != nullptr)
+	if (OnlineSub != nullptr)
 	{
 		TSharedPtr<FOnlineUserPresence> Presence;
 		TSharedPtr<FUniqueNetId> UserId = OnlineIdentity->GetUniquePlayerId(0);
 		if(UserId.IsValid())
 		{
-			OnlineSubMcp->GetPresenceInterface()->GetCachedPresence(*UserId, Presence);
+			OnlineSub->GetPresenceInterface()->GetCachedPresence(*UserId, Presence);
 			if(Presence.IsValid())
 			{
 				return Presence->Status.State;
@@ -863,8 +856,8 @@ void FFriendsAndChatManager::OnQueryRecentPlayersComplete(const FUniqueNetId& Us
 			
 			if(QueryUserIds.Num())
 			{
-				check(OnlineSubMcp != nullptr && OnlineSubMcp->GetMcpPersonaService().IsValid());
-				IOnlineUserPtr UserInterface = OnlineSubMcp->GetUserInterface();
+				check(OnlineSub != nullptr && OnlineSub->GetUserInterface().IsValid());
+				IOnlineUserPtr UserInterface = OnlineSub->GetUserInterface();
 				UserInterface->QueryUserInfo(0, QueryUserIds);
 				bFoundAllIds = false;
 			}
@@ -908,8 +901,8 @@ void FFriendsAndChatManager::PreProcessList(const FString& ListName)
 			}
 		}
 
-		check(OnlineSubMcp != nullptr && OnlineSubMcp->GetMcpPersonaService().IsValid());
-		IOnlineUserPtr UserInterface = OnlineSubMcp->GetUserInterface();
+		check(OnlineSub != nullptr && OnlineSub->GetUserInterface().IsValid());
+		IOnlineUserPtr UserInterface = OnlineSub->GetUserInterface();
 
 		if ( QueryUserIds.Num() > 0 )
 		{
@@ -933,13 +926,13 @@ void FFriendsAndChatManager::OnQueryUserInfoComplete( int32 LocalPlayer, bool bW
 {
 	if( ManagerState == EFriendsAndManagerState::RequestingFriendsList)
 	{
-		check(OnlineSubMcp != nullptr && OnlineSubMcp->GetMcpPersonaService().IsValid());
-		IOnlineUserPtr UserInterface = OnlineSubMcp->GetUserInterface();
+		check(OnlineSub != nullptr && OnlineSub->GetUserInterface().IsValid());
+		IOnlineUserPtr UserInterface = OnlineSub->GetUserInterface();
 
 		for ( int32 UserIdx=0; UserIdx < UserIds.Num(); UserIdx++ )
 		{
 			TSharedPtr<FOnlineFriend> OnlineFriend = FriendsInterface->GetFriend( 0, *UserIds[UserIdx], EFriendsLists::ToString( EFriendsLists::Default ) );
-			TSharedPtr<FOnlineUser> OnlineUser = OnlineSubMcp->GetUserInterface()->GetUserInfo( LocalPlayer, *UserIds[UserIdx] );
+			TSharedPtr<FOnlineUser> OnlineUser = OnlineSub->GetUserInterface()->GetUserInfo( LocalPlayer, *UserIds[UserIdx] );
 			if (OnlineFriend.IsValid() && OnlineUser.IsValid())
 			{
 				TSharedPtr<IFriendItem> Existing;
@@ -984,7 +977,7 @@ void FFriendsAndChatManager::OnQueryUserInfoComplete( int32 LocalPlayer, bool bW
 				for (const auto& RecentPlayer : Players)
 				{
 					TSharedRef<FFriendRecentPlayerItem> RecentPlayerItem = MakeShareable(new FFriendRecentPlayerItem(RecentPlayer));
-					TSharedPtr<FOnlineUser> OnlineUser = OnlineSubMcp->GetUserInterface()->GetUserInfo(LocalPlayer, *RecentPlayer->GetUserId());
+					TSharedPtr<FOnlineUser> OnlineUser = OnlineSub->GetUserInterface()->GetUserInfo(LocalPlayer, *RecentPlayer->GetUserId());
 					RecentPlayerItem->SetOnlineUser(OnlineUser);
 					RecentPlayersList.Add(RecentPlayerItem);
 				}
@@ -1079,14 +1072,18 @@ void FFriendsAndChatManager::RefreshList()
 void FFriendsAndChatManager::SendFriendRequests()
 {
 	// Invite Friends
-	FOnlinePersonaMcpPtr OnlinePersonaMcp = OnlineSubMcp->GetMcpPersonaService();
-	TSharedPtr<FUniqueNetId> UserId = OnlineIdentity->GetUniquePlayerId(0);
-	if (UserId.IsValid())
+	IOnlineUserPtr UserInterface = OnlineSub->GetUserInterface();
+	if (UserInterface.IsValid())
 	{
-		for (int32 Index = 0; Index < FriendByNameRequests.Num(); Index++)
-		{
-			OnlinePersonaMcp->QueryUserIdMapping(*UserId, FriendByNameRequests[Index]);
-		}
+	    TSharedPtr<FUniqueNetId> UserId = OnlineIdentity->GetUniquePlayerId(0);
+	    if (UserId.IsValid())
+	    {
+		    for (int32 Index = 0; Index < FriendByNameRequests.Num(); Index++)
+		    {
+				    UserInterface->QueryUserIdMapping(*UserId, FriendByNameRequests[Index], OnQueryUserIdMappingCompleteDelegate);
+			    }
+		    }
+	    }
 	}
 }
 
@@ -1175,8 +1172,8 @@ void FFriendsAndChatManager::SendInviteAcceptedNotification(const TSharedPtr< IF
 
 const FText FFriendsAndChatManager::GetInviteNotificationText(TSharedPtr< IFriendItem > Friend) const
 {
-	FFormatNamedArguments Args;
-	Args.Add(TEXT("Username"), FText::FromString(Friend->GetName()));
+		FFormatNamedArguments Args;
+		Args.Add(TEXT("Username"), FText::FromString(Friend->GetName()));
 
 	if(Friend->IsPendingAccepted())
 	{
@@ -1187,7 +1184,7 @@ const FText FFriendsAndChatManager::GetInviteNotificationText(TSharedPtr< IFrien
 
 void FFriendsAndChatManager::OnQueryUserIdMappingComplete(bool bWasSuccessful, const FUniqueNetId& RequestingUserId, const FString& DisplayName, const FUniqueNetId& IdentifiedUserId, const FString& Error)
 {
-	check( OnlineSubMcp != nullptr && OnlineSubMcp->GetMcpPersonaService().IsValid() );
+	check( OnlineSub != nullptr && OnlineSub->GetUserInterface().IsValid() );
 
 	EFindFriendResult::Type FindFriendResult = EFindFriendResult::NotFound;
 
@@ -1397,10 +1394,10 @@ void FFriendsAndChatManager::OnFriendInviteReceived(const FUniqueNetId& UserId,
 
 void FFriendsAndChatManager::OnGameInviteReceived(const FUniqueNetId& UserId, const FUniqueNetId& FromId, const FOnlineSessionSearchResult& InviteResult)
 {
-	if (OnlineSubMcp != NULL &&
-		OnlineSubMcp->GetIdentityInterface().IsValid())
+	if (OnlineSub != NULL &&
+		OnlineSub->GetIdentityInterface().IsValid())
 	{
-		TSharedPtr<FUniqueNetId> FromIdPtr = OnlineSubMcp->GetIdentityInterface()->CreateUniquePlayerId(FromId.ToString());
+		TSharedPtr<FUniqueNetId> FromIdPtr = OnlineSub->GetIdentityInterface()->CreateUniquePlayerId(FromId.ToString());
 		if (FromIdPtr.IsValid())
 		{
 			ReceivedGameInvites.AddUnique(FReceivedGameInvite(FromIdPtr.ToSharedRef(), InviteResult));
@@ -1410,8 +1407,8 @@ void FFriendsAndChatManager::OnGameInviteReceived(const FUniqueNetId& UserId, co
 
 void FFriendsAndChatManager::ProcessReceivedGameInvites()
 {
-	if (OnlineSubMcp != NULL &&
-		OnlineSubMcp->GetUserInterface().IsValid())
+	if (OnlineSub != NULL &&
+		OnlineSub->GetUserInterface().IsValid())
 	{
 		for (int32 Idx = 0; Idx < ReceivedGameInvites.Num(); Idx++)
 		{
@@ -1425,7 +1422,7 @@ void FFriendsAndChatManager::ProcessReceivedGameInvites()
 			}
 			if (!UserInfo.IsValid())
 			{
-				UserInfo = OnlineSubMcp->GetUserInterface()->GetUserInfo(0, *Invite.FromId);
+				UserInfo = OnlineSub->GetUserInterface()->GetUserInfo(0, *Invite.FromId);
 			}
 			if (UserInfo.IsValid())
 			{
@@ -1449,8 +1446,8 @@ bool FFriendsAndChatManager::RequestGameInviteUserInfo()
 	bool bPending = false;
 
 	// query for user ids that are not already cached from game invites
-	IOnlineUserPtr UserInterface = OnlineSubMcp->GetUserInterface();
-	IOnlineIdentityPtr IdentityInterface = OnlineSubMcp->GetIdentityInterface();
+	IOnlineUserPtr UserInterface = OnlineSub->GetUserInterface();
+	IOnlineIdentityPtr IdentityInterface = OnlineSub->GetIdentityInterface();
 	if (UserInterface.IsValid() &&
 		IdentityInterface.IsValid())
 	{
@@ -1538,16 +1535,16 @@ void FFriendsAndChatManager::SendGameInvite(const TSharedPtr<IFriendItem>& Frien
 
 void FFriendsAndChatManager::SendGameInvite(const FUniqueNetId& ToUser)
 {
-	if (OnlineSubMcp != nullptr &&
+	if (OnlineSub != nullptr &&
 		OnlineIdentity.IsValid() &&
-		OnlineSubMcp->GetSessionInterface().IsValid() &&
-		OnlineSubMcp->GetSessionInterface()->GetNamedSession(GameSessionName) != nullptr)
+		OnlineSub->GetSessionInterface().IsValid() &&
+		OnlineSub->GetSessionInterface()->GetNamedSession(GameSessionName) != nullptr)
 	{
 		TSharedPtr<FUniqueNetId> UserId = OnlineIdentity->GetUniquePlayerId(0);
 		if (UserId.IsValid())
 		{
 			OnlineSubMcp->GetSessionInterface()->SendSessionInviteToFriend(*UserId, GameSessionName, ToUser);
-			AddFriendsToast(LOCTEXT("InviteToGameSentToast", "Invite Sent"));
+
 			Analytics.RecordGameInvite(ToUser, TEXT("Social.GameInvite.Send"));
 		}
 	}
@@ -1625,7 +1622,7 @@ void FFriendsAndChatAnalytics::RecordGameInvite(const FUniqueNetId& ToUser, cons
 {
 	if (Provider.IsValid())
 	{
-		IOnlineIdentityPtr OnlineIdentity = Online::GetIdentityInterface(MCP_SUBSYSTEM);
+		IOnlineIdentityPtr OnlineIdentity = Online::GetIdentityInterface(TEXT("MCP"));
 		if (OnlineIdentity.IsValid())
 		{
 			TSharedPtr<FUniqueNetId> UserId = OnlineIdentity->GetUniquePlayerId(0);
@@ -1645,7 +1642,7 @@ void FFriendsAndChatAnalytics::RecordFriendAction(const IFriendItem& Friend, con
 {
 	if (Provider.IsValid())
 	{
-		IOnlineIdentityPtr OnlineIdentity = Online::GetIdentityInterface(MCP_SUBSYSTEM);
+		IOnlineIdentityPtr OnlineIdentity = Online::GetIdentityInterface(TEXT("MCP"));
 		if (OnlineIdentity.IsValid())
 		{
 			TSharedPtr<FUniqueNetId> UserId = OnlineIdentity->GetUniquePlayerId(0);
@@ -1665,7 +1662,7 @@ void FFriendsAndChatAnalytics::RecordAddFriend(const FString& FriendName, const
 {
 	if (Provider.IsValid())
 	{
-		IOnlineIdentityPtr OnlineIdentity = Online::GetIdentityInterface(MCP_SUBSYSTEM);
+		IOnlineIdentityPtr OnlineIdentity = Online::GetIdentityInterface(TEXT("MCP"));
 		if (OnlineIdentity.IsValid())
 		{
 			TSharedPtr<FUniqueNetId> UserId = OnlineIdentity->GetUniquePlayerId(0);
@@ -1688,7 +1685,7 @@ void FFriendsAndChatAnalytics::RecordToggleChat(const FString& Channel, bool bEn
 {
 	if (Provider.IsValid())
 	{
-		IOnlineIdentityPtr OnlineIdentity = Online::GetIdentityInterface(MCP_SUBSYSTEM);
+		IOnlineIdentityPtr OnlineIdentity = Online::GetIdentityInterface(TEXT("MCP"));
 		if (OnlineIdentity.IsValid())
 		{
 			TSharedPtr<FUniqueNetId> UserId = OnlineIdentity->GetUniquePlayerId(0);
@@ -1723,7 +1720,7 @@ void FFriendsAndChatAnalytics::FlushChatStats()
 {
 	if (Provider.IsValid())
 	{
-		IOnlineIdentityPtr OnlineIdentity = Online::GetIdentityInterface(MCP_SUBSYSTEM);
+		IOnlineIdentityPtr OnlineIdentity = Online::GetIdentityInterface(TEXT("MCP"));
 		if (OnlineIdentity.IsValid())
 		{
 			TSharedPtr<FUniqueNetId> UserId = OnlineIdentity->GetUniquePlayerId(0);
@@ -1743,7 +1740,7 @@ void FFriendsAndChatAnalytics::FlushChatStats()
 
 void FFriendsAndChatAnalytics::AddPresenceAttributes(const FUniqueNetId& UserId, TArray<FAnalyticsEventAttribute>& Attributes) const
 {
-	IOnlinePresencePtr OnlinePresence = Online::GetPresenceInterface(MCP_SUBSYSTEM);
+	IOnlinePresencePtr OnlinePresence = Online::GetPresenceInterface(TEXT("MCP"));
 	if (OnlinePresence.IsValid())
 	{
 		TSharedPtr<FOnlineUserPresence> Presence;
diff --git a/Engine/Source/Developer/FriendsAndChat/Private/FriendsAndChatManager.h b/Engine/Source/Developer/FriendsAndChat/Private/FriendsAndChatManager.h
--- a/Engine/Source/Developer/FriendsAndChat/Private/FriendsAndChatManager.h
+++ b/Engine/Source/Developer/FriendsAndChat/Private/FriendsAndChatManager.h
@@ -685,7 +685,7 @@ private:
 	// Delegate to use for deleting a friend
 	FOnDeleteFriendCompleteDelegate OnDeleteFriendCompleteDelegate;
 	// Delegate for querying user id from a name string
-	FOnQueryUserIdMappingCompleteDelegate OnQueryUserIdMappingCompleteDelegate;
+	IOnlineUser::FOnQueryUserMappingComplete OnQueryUserIdMappingCompleteDelegate;
 	// Delegate to use for querying user info list
 	FOnQueryUserInfoCompleteDelegate OnQueryUserInfoCompleteDelegate;
 	// Delegate to use for querying user presence
@@ -738,7 +738,7 @@ private:
 	// Holds the Friends Interface
 	IOnlineFriendsPtr FriendsInterface;
 	// Holds the Online Subsystem
-	FOnlineSubsystemMcp* OnlineSubMcp;
+	IOnlineSubsystem* OnlineSub;
 
 	/* Chat stuff
 	*****************************************************************************/
@@ -806,6 +806,5 @@ private:
 	FDelegateHandle OnPresenceReceivedCompleteDelegateHandle;
 	FDelegateHandle OnGameInviteReceivedDelegateHandle;
 	FDelegateHandle OnDestroySessionCompleteDelegateHandle;
-	FDelegateHandle OnQueryUserIdMappingCompleteDelegateHandle;
 	FDelegateHandle UpdateFriendsTickerDelegateHandle;
 };
diff --git a/Engine/Source/Developer/FriendsAndChat/Private/FriendsAndChatPrivatePCH.h b/Engine/Source/Developer/FriendsAndChat/Private/FriendsAndChatPrivatePCH.h
--- a/Engine/Source/Developer/FriendsAndChat/Private/FriendsAndChatPrivatePCH.h
+++ b/Engine/Source/Developer/FriendsAndChat/Private/FriendsAndChatPrivatePCH.h
@@ -17,9 +17,7 @@
 #include "FriendsAndChat.h"
 #include "Online.h"
 #include "OnlineSubsystem.h"
-#include "OnlineSubsystemMcp.h"
 #include "OnlineFriendsInterface.h"
-#include "OnlinePersonaMcp.h"
 #include "OnlineUserInterface.h"
 #include "OnlineMessageInterface.h"
 #include "OnlinePresenceInterface.h"
diff --git a/Engine/Source/Developer/FriendsAndChat/Public/FriendsAndChat.h b/Engine/Source/Developer/FriendsAndChat/Public/FriendsAndChat.h
--- a/Engine/Source/Developer/FriendsAndChat/Public/FriendsAndChat.h
+++ b/Engine/Source/Developer/FriendsAndChat/Public/FriendsAndChat.h
@@ -12,7 +12,6 @@
 #include "SlateBasics.h"
 #include "SlateStyle.h"
 #include "OnlineSubsystem.h"
-#include "OnlineSubsystemMcp.h"
 
 
 /* Interfaces
diff --git a/Engine/Source/Runtime/Online/OnlineSubsystem/Public/Interfaces/OnlineUserInterface.h b/Engine/Source/Runtime/Online/OnlineSubsystem/Public/Interfaces/OnlineUserInterface.h
--- a/Engine/Source/Runtime/Online/OnlineSubsystem/Public/Interfaces/OnlineUserInterface.h
+++ b/Engine/Source/Runtime/Online/OnlineSubsystem/Public/Interfaces/OnlineUserInterface.h
@@ -62,4 +62,25 @@ public:
 	 * @return user info or null ptr if not found
 	 */
 	virtual TSharedPtr<FOnlineUser> GetUserInfo(int32 LocalUserNum, const class FUniqueNetId& UserId) = 0;
+
+	/**
+	 * Called when done querying for a UserId mapping from a requested display name
+	 *
+	 * @param bWasSuccessful true if server was contacted and a valid result received
+	 * @param UserId user id initiating the request
+	 * @param DisplayNameOrEmail the name string that was being queried
+	 * @param FoundUserId the user id matched for the passed name string
+	 * @param Error string representing the error condition
+	 */
+	DECLARE_DELEGATE_FiveParams(FOnQueryUserMappingComplete, bool /*bWasSuccessful*/, const FUniqueNetId& /*UserId*/, const FString& /*DisplayNameOrEmail*/, const FUniqueNetId& /*FoundUserId*/, const FString& /*Error*/);
+
+	/**
+	 * Contacts server to obtain a user id from an arbitrary user-entered name string, eg. display name
+	 *
+	 * @param UserId id of the user that is requesting the name string lookup
+	 * @param DisplayNameOrEmail a string of a display name or email to attempt to map to a user id
+	 *
+	 * @return true if the operation was started successfully
+	 */
+	virtual bool QueryUserIdMapping(const FUniqueNetId& UserId, const FString& DisplayNameOrEmail, const FOnQueryUserMappingComplete& Delegate = FOnQueryUserMappingComplete()) = 0;	
 };
diff --git a/Engine/Source/Runtime/Online/OnlineSubsystemFacebook/Private/IOS/OnlineUserFacebook.cpp b/Engine/Source/Runtime/Online/OnlineSubsystemFacebook/Private/IOS/OnlineUserFacebook.cpp
--- a/Engine/Source/Runtime/Online/OnlineSubsystemFacebook/Private/IOS/OnlineUserFacebook.cpp
+++ b/Engine/Source/Runtime/Online/OnlineSubsystemFacebook/Private/IOS/OnlineUserFacebook.cpp
@@ -156,3 +156,9 @@ TSharedPtr<FOnlineUser> FOnlineUserFacebook::GetUserInfo(int32 LocalUserNum, con
 	return Result;
 }
 
+bool FOnlineUserFacebook::QueryUserIdMapping(const FUniqueNetId& UserId, const FString& DisplayNameOrEmail, const FOnQueryUserMappingComplete& Delegate)
+{
+	Delegate.ExecuteIfBound(false, UserId, DisplayNameOrEmail, FUniqueNetIdString(), TEXT("not implemented"));
+	return false;
+}
+
diff --git a/Engine/Source/Runtime/Online/OnlineSubsystemFacebook/Private/IOS/OnlineUserFacebook.h b/Engine/Source/Runtime/Online/OnlineSubsystemFacebook/Private/IOS/OnlineUserFacebook.h
--- a/Engine/Source/Runtime/Online/OnlineSubsystemFacebook/Private/IOS/OnlineUserFacebook.h
+++ b/Engine/Source/Runtime/Online/OnlineSubsystemFacebook/Private/IOS/OnlineUserFacebook.h
@@ -77,7 +77,8 @@ public:
 
 	virtual bool QueryUserInfo(int32 LocalUserNum, const TArray<TSharedRef<class FUniqueNetId> >& UserIds) override;
 	virtual bool GetAllUserInfo(int32 LocalUserNum, TArray< TSharedRef<class FOnlineUser> >& OutUsers) override;
-	virtual TSharedPtr<FOnlineUser> GetUserInfo(int32 LocalUserNum, const class FUniqueNetId& UserId) override;
+	virtual TSharedPtr<FOnlineUser> GetUserInfo(int32 LocalUserNum, const class FUniqueNetId& UserId) override;	
+	virtual bool QueryUserIdMapping(const FUniqueNetId& UserId, const FString& DisplayNameOrEmail, const FOnQueryUserMappingComplete& Delegate = FOnQueryUserMappingComplete()) override;
 
 	// FOnlineUserFacebook
 
