Fix build error.

diff --git a/Engine/Source/Developer/FriendsAndChat/Private/FriendsAndChatManager.cpp b/Engine/Source/Developer/FriendsAndChat/Private/FriendsAndChatManager.cpp
--- a/Engine/Source/Developer/FriendsAndChat/Private/FriendsAndChatManager.cpp
+++ b/Engine/Source/Developer/FriendsAndChat/Private/FriendsAndChatManager.cpp
@@ -1080,8 +1080,7 @@ void FFriendsAndChatManager::SendFriendRequests()
 	    {
 		    for (int32 Index = 0; Index < FriendByNameRequests.Num(); Index++)
 		    {
-				    UserInterface->QueryUserIdMapping(*UserId, FriendByNameRequests[Index], OnQueryUserIdMappingCompleteDelegate);
-			    }
+			    UserInterface->QueryUserIdMapping(*UserId, FriendByNameRequests[Index], OnQueryUserIdMappingCompleteDelegate);
 		    }
 	    }
 	}
@@ -1543,7 +1542,7 @@ void FFriendsAndChatManager::SendGameInvite(const FUniqueNetId& ToUser)
 		TSharedPtr<FUniqueNetId> UserId = OnlineIdentity->GetUniquePlayerId(0);
 		if (UserId.IsValid())
 		{
-			OnlineSubMcp->GetSessionInterface()->SendSessionInviteToFriend(*UserId, GameSessionName, ToUser);
+			OnlineSub->GetSessionInterface()->SendSessionInviteToFriend(*UserId, GameSessionName, ToUser);
 
 			Analytics.RecordGameInvite(ToUser, TEXT("Social.GameInvite.Send"));
 		}
