Merging UE4-LauncherDev UE4 up until 2421209

This merges all changes from the LauncherDev branch up until before the failed Release4.7 into LauncherDev, which broke parts of branch history. All other changes since then will be spot merged after this check-in.

diff --git a/Engine/Source/Developer/FriendsAndChat/Private/FriendsAndChatManager.cpp b/Engine/Source/Developer/FriendsAndChat/Private/FriendsAndChatManager.cpp
--- a/Engine/Source/Developer/FriendsAndChat/Private/FriendsAndChatManager.cpp
+++ b/Engine/Source/Developer/FriendsAndChat/Private/FriendsAndChatManager.cpp
@@ -16,6 +16,13 @@
 
 #define LOCTEXT_NAMESPACE "FriendsAndChatManager"
 
+namespace FriendsAndChatManagerDefs
+{
+	static const FVector2D FriendsListWindowPosition(100, 100);
+	static const FVector2D FriendsListWindowSize(400, 458);
+	static const FVector2D ChatWindowSize(420, 500);
+	static const float FriendsAndChatWindowGap = 20.0f;
+}
 
 /* FFriendsAndChatManager structors
  *****************************************************************************/
@@ -225,8 +232,6 @@ void FFriendsAndChatManager::JoinPublicChatRoom(const FString& RoomName)
 
 void FFriendsAndChatManager::CreateFriendsListWindow(const FFriendsAndChatStyle* InStyle)
 {
-	const FVector2D DEFAULT_WINDOW_SIZE = FVector2D(400, 458);
-
 	Style = *InStyle;
 	FFriendsAndChatModuleStyle::Initialize(Style);
 
@@ -235,24 +240,25 @@ void FFriendsAndChatManager::CreateFriendsListWindow(const FFriendsAndChatStyle*
 		FriendWindow = SNew( SWindow )
 		.Title(LOCTEXT("FFriendsAndChatManager_FriendsTitle", "Friends List"))
 		.Style(&Style.WindowStyle)
-		.ClientSize(DEFAULT_WINDOW_SIZE)
-		.ScreenPosition(FVector2D(100, 100))
+		.ClientSize(FriendsAndChatManagerDefs::FriendsListWindowSize)
+		.ScreenPosition(FriendsAndChatManagerDefs::FriendsListWindowPosition)
 		.AutoCenter( EAutoCenter::None )
 		.SizingRule(ESizingRule::UserSized)
 		.SupportsMaximize(true)
 		.SupportsMinimize(true)
 		.bDragAnywhere(true)
-		.CreateTitleBar(false);
+		.CreateTitleBar(false)
+		.LayoutBorder(FMargin(0));
 
 		FriendWindow->SetOnWindowClosed(FOnWindowClosed::CreateRaw(this, &FFriendsAndChatManager::HandleFriendsWindowClosed));
 
-		BuildFriendsUI();
+		BuildFriendsUI(FriendWindow);
 		FriendWindow = FSlateApplication::Get().AddWindow(FriendWindow.ToSharedRef());
 	}
 	else if(FriendWindow->IsWindowMinimized())
 	{
 		FriendWindow->Restore();
-		BuildFriendsUI();
+		BuildFriendsUI(FriendWindow);
 	}
 
 	FriendWindow->BringToFront(true);
@@ -266,12 +272,12 @@ void FFriendsAndChatManager::HandleFriendsWindowClosed(const TSharedRef<SWindow>
 	FriendWindow.Reset();
 }
 
-void FFriendsAndChatManager::BuildFriendsUI()
+void FFriendsAndChatManager::BuildFriendsUI(TSharedPtr< SWindow > WindowPtr)
 {
-	check(FriendWindow.IsValid());
+	check(WindowPtr.IsValid());
 
 	TSharedPtr<SWindowTitleBar> TitleBar;
-	FriendWindow->SetContent(
+	WindowPtr->SetContent(
 		SNew(SBorder)
 		.BorderImage( &Style.Background )
 		.VAlign( VAlign_Fill )
@@ -285,7 +291,7 @@ void FFriendsAndChatManager::BuildFriendsUI()
 				+ SVerticalBox::Slot()
 				.AutoHeight()
 				[
-					SAssignNew(TitleBar, SWindowTitleBar, FriendWindow.ToSharedRef(), nullptr, HAlign_Center)
+					SAssignNew(TitleBar, SWindowTitleBar, WindowPtr.ToSharedRef(), nullptr, HAlign_Center)
 					.Style(&Style.WindowStyle)
 					.ShowAppIcon(false)
 					.Title(FText::GetEmpty())
@@ -367,7 +373,12 @@ TSharedPtr<IChatViewModel> FFriendsAndChatManager::GetChatViewModel()
 
 void FFriendsAndChatManager::CreateChatWindow(const struct FFriendsAndChatStyle* InStyle)
 {
-	const FVector2D DEFAULT_WINDOW_SIZE = FVector2D(420, 500);
+	FVector2D WindowPosition
+	(
+		FriendsAndChatManagerDefs::FriendsListWindowPosition.X + FriendsAndChatManagerDefs::FriendsListWindowSize.X + FriendsAndChatManagerDefs::FriendsAndChatWindowGap,
+		FriendsAndChatManagerDefs::FriendsListWindowPosition.Y
+	);
+
 	check(MessageManager.IsValid());
 	bCreateChatWindow = true;
 
@@ -379,13 +390,14 @@ void FFriendsAndChatManager::CreateChatWindow(const struct FFriendsAndChatStyle*
 		ChatWindow = SNew( SWindow )
 		.Title( LOCTEXT( "FriendsAndChatManager_ChatTitle", "Chat Window") )
 		.Style(&Style.WindowStyle)
-		.ClientSize( DEFAULT_WINDOW_SIZE )
-		.ScreenPosition( FVector2D( 200, 100 ) )
+		.ClientSize(FriendsAndChatManagerDefs::ChatWindowSize)
+		.ScreenPosition(WindowPosition)
 		.AutoCenter( EAutoCenter::None )
 		.SupportsMaximize( true )
 		.SupportsMinimize( true )
 		.CreateTitleBar( false )
-		.SizingRule(ESizingRule::UserSized);
+		.SizingRule(ESizingRule::UserSized)
+		.LayoutBorder(FMargin(0));
 
 		ChatWindow->SetOnWindowClosed(FOnWindowClosed::CreateRaw(this, &FFriendsAndChatManager::HandleChatWindowClosed));
 
diff --git a/Engine/Source/Developer/FriendsAndChat/Private/FriendsAndChatManager.h b/Engine/Source/Developer/FriendsAndChat/Private/FriendsAndChatManager.h
--- a/Engine/Source/Developer/FriendsAndChat/Private/FriendsAndChatManager.h
+++ b/Engine/Source/Developer/FriendsAndChat/Private/FriendsAndChatManager.h
@@ -402,7 +402,7 @@ private:
 	void RefreshList();
 
 	/** Build the friends UI. */
-	void BuildFriendsUI();
+	void BuildFriendsUI(TSharedPtr< SWindow > WindowPtr);
 
 	/**
 	 * Set the manager state.
diff --git a/Engine/Source/Runtime/Online/BuildPatchServices/Private/BuildPatchHTTP.cpp b/Engine/Source/Runtime/Online/BuildPatchServices/Private/BuildPatchHTTP.cpp
--- a/Engine/Source/Runtime/Online/BuildPatchServices/Private/BuildPatchHTTP.cpp
+++ b/Engine/Source/Runtime/Online/BuildPatchServices/Private/BuildPatchHTTP.cpp
@@ -32,7 +32,7 @@ bool FBuildPatchHTTP::Tick( float Delta )
 			{
 				FHttpRequestInfo HttpRequestInfo = HttpRequestQueue.Pop();
 				TSharedRef< IHttpRequest > HttpRequest = FHttpModule::Get().CreateRequest();
-				HttpRequest->OnProcessRequestComplete() = HttpRequestInfo.OnCompleteDelegate;
+				HttpRequest->OnProcessRequestComplete().BindSP(this, &FBuildPatchHTTP::HttpRequestComplete, HttpRequestInfo);
 				HttpRequest->OnRequestProgress() = HttpRequestInfo.OnProgressDelegate;
 				HttpRequest->SetURL( HttpRequestInfo.UrlRequest );
 				HttpRequest->ProcessRequest();
@@ -66,6 +66,13 @@ bool FBuildPatchHTTP::Tick( float Delta )
 	return true;
 }
 
+void FBuildPatchHTTP::HttpRequestComplete(FHttpRequestPtr Request, FHttpResponsePtr Response, bool bSucceeded, FHttpRequestInfo HttpRequestInfo)
+{
+	// Remove this request from our map
+	HttpRequestMap.Remove(HttpRequestInfo.RequestID);
+	HttpRequestInfo.OnCompleteDelegate.ExecuteIfBound(Request, Response, bSucceeded);
+}
+
 int32 FBuildPatchHTTP::QueueHttpRequest( const FString& UrlRequest, const FHttpRequestCompleteDelegate& OnCompleteDelegate, const FHttpRequestProgressDelegate& OnProgressDelegate )
 {
 	if( !FBuildPatchInstallError::HasFatalError() )
@@ -128,24 +135,23 @@ void FBuildPatchHTTP::QueueAnalyticsEvent( const FString& EventName, const TArra
 
 /* Static initialization
  *****************************************************************************/
-FBuildPatchHTTP* FBuildPatchHTTP::SingletonInstance = NULL;
+TSharedPtr<FBuildPatchHTTP> FBuildPatchHTTP::SingletonInstance = NULL;
 FBuildPatchHTTP& FBuildPatchHTTP::Get()
 {
-	if( SingletonInstance == NULL )
+	if (SingletonInstance.IsValid() == false)
 	{
-		SingletonInstance = new FBuildPatchHTTP();
-		TickDelegateHandle = FTicker::GetCoreTicker().AddTicker( FTickerDelegate::CreateRaw( SingletonInstance, &FBuildPatchHTTP::Tick ) );
+		SingletonInstance = MakeShareable(new FBuildPatchHTTP());
+		TickDelegateHandle = FTicker::GetCoreTicker().AddTicker( FTickerDelegate::CreateSP( SingletonInstance.Get(), &FBuildPatchHTTP::Tick ) );
 	}
-	return *SingletonInstance;
+	return *SingletonInstance.Get();
 }
 void FBuildPatchHTTP::OnShutdown()
 {
-	if( SingletonInstance != NULL )
+	if (SingletonInstance.IsValid())
 	{
 		SingletonInstance->Tick( 0.0f );
 		FTicker::GetCoreTicker().RemoveTicker( TickDelegateHandle );
-		delete SingletonInstance;
-		SingletonInstance = NULL;
+		SingletonInstance.Reset();
 	}
 }
 FDelegateHandle FBuildPatchHTTP::TickDelegateHandle;
diff --git a/Engine/Source/Runtime/Online/BuildPatchServices/Private/BuildPatchHTTP.h b/Engine/Source/Runtime/Online/BuildPatchServices/Private/BuildPatchHTTP.h
--- a/Engine/Source/Runtime/Online/BuildPatchServices/Private/BuildPatchHTTP.h
+++ b/Engine/Source/Runtime/Online/BuildPatchServices/Private/BuildPatchHTTP.h
@@ -13,7 +13,7 @@
 /**
  * Implements the BuildPatchHTTP
  */
-class FBuildPatchHTTP : private FNoncopyable
+class FBuildPatchHTTP : public TSharedFromThis<FBuildPatchHTTP>
 {
 private:
 
@@ -114,11 +114,21 @@ private:
 	 */
 	bool Tick( float Delta );
 
+	/**
+	 * Internal HTTP request complete function, used to clean out this request from the map and remove any pending cancels which
+	 * could cause double complete delegates
+	 * @param Request			The http request interface
+	 * @param Response			The http response interface
+	 * @param bSucceeded		Whether the request was successful
+	 * @param HttpRequestInfo	The original request details
+	 */
+	void HttpRequestComplete(FHttpRequestPtr Request, FHttpResponsePtr Response, bool bSucceeded, FHttpRequestInfo HttpRequestInfo);
+
 /* Here we have private static access for the singleton
 *****************************************************************************/
 private:
 	static FBuildPatchHTTP& Get();
-	static FBuildPatchHTTP* SingletonInstance;
+	static TSharedPtr<FBuildPatchHTTP> SingletonInstance;
 	static FDelegateHandle TickDelegateHandle;
 };
 
diff --git a/Engine/Source/Runtime/Online/BuildPatchServices/Private/BuildPatchManifest.cpp b/Engine/Source/Runtime/Online/BuildPatchServices/Private/BuildPatchManifest.cpp
--- a/Engine/Source/Runtime/Online/BuildPatchServices/Private/BuildPatchManifest.cpp
+++ b/Engine/Source/Runtime/Online/BuildPatchServices/Private/BuildPatchManifest.cpp
@@ -968,6 +968,7 @@ bool FBuildPatchAppManifest::DeserializeFromJSON( const FString& JSONInput )
 	}
 
 	// Get the ChunkFilesizeList
+	bool bHasChunkFilesizeList = false;
 	TSharedPtr< FJsonValue > JsonChunkFilesizeList = JsonValueMap.FindRef(TEXT("ChunkFilesizeList"));
 	if (JsonChunkFilesizeList.IsValid())
 	{
@@ -983,11 +984,12 @@ bool FBuildPatchAppManifest::DeserializeFromJSON( const FString& JSONInput )
 				{
 					FChunkInfoData* ChunkInfoData = ChunkInfoLookup[ChunkGuid];
 					ChunkInfoData->FileSize = ChunkSize;
+					bHasChunkFilesizeList = true;
 				}
 			}
 		}
 	}
-	else
+	if (bHasChunkFilesizeList == false)
 	{
 		// Missing chunk list, version before we saved them compressed.. Assume chunk size
 		for (FChunkInfoData& ChunkInfo : Data->ChunkList)
diff --git a/Engine/Source/Runtime/Slate/Private/Widgets/Layout/SScrollBox.cpp b/Engine/Source/Runtime/Slate/Private/Widgets/Layout/SScrollBox.cpp
--- a/Engine/Source/Runtime/Slate/Private/Widgets/Layout/SScrollBox.cpp
+++ b/Engine/Source/Runtime/Slate/Private/Widgets/Layout/SScrollBox.cpp
@@ -379,8 +379,8 @@ bool SScrollBox::ScrollDescendantIntoView(const FGeometry& MyGeometry, const TSh
 	FArrangedWidget* WidgetGeometry = Result.Find( WidgetToFind.ToSharedRef() );
 	if ( ensureMsg( WidgetGeometry, TEXT("Unable to scroll to descendant as it's not a child of the scrollbox") ) )
 	{
-		// Clear any existing scroll offset, since we're working with absolute positions
-		SetScrollOffset(0);
+		// @todo: This is a workaround because DesiredScrollOffset can exceed the ScrollMax when mouse dragging on the scroll bar and we need it clamped here or the offset is wrong
+		ScrollBy(MyGeometry, 0, false);
 
 		// Calculate how much we would need to scroll to bring this to the top/left of the scroll box
 		const float WidgetPosition = GetScrollComponentFromVector(WidgetGeometry->Geometry.AbsolutePosition);
