[INTEGRATE] merge ue4-launcherdev->ue4 cl 2438986

diff --git a/Engine/Source/Developer/FriendsAndChat/Private/Models/FriendItem.cpp b/Engine/Source/Developer/FriendsAndChat/Private/Models/FriendItem.cpp
--- a/Engine/Source/Developer/FriendsAndChat/Private/Models/FriendItem.cpp
+++ b/Engine/Source/Developer/FriendsAndChat/Private/Models/FriendItem.cpp
@@ -164,6 +164,12 @@ bool FFriendItem::IsGameJoinable() const
 	return false;
 }
 
+bool FFriendItem::CanInvite() const
+{
+	FString FriendsClientID = GetClientId();
+	return FriendsClientID == FFriendsAndChatManager::Get()->GetUserClientId() || FriendsClientID == FFriendItem::LauncherClientId;
+}
+
 FString FFriendItem::GetGameSessionId() const
 {
 	FString SessionId;
diff --git a/Engine/Source/Developer/FriendsAndChat/Private/Models/FriendItem.h b/Engine/Source/Developer/FriendsAndChat/Private/Models/FriendItem.h
--- a/Engine/Source/Developer/FriendsAndChat/Private/Models/FriendItem.h
+++ b/Engine/Source/Developer/FriendsAndChat/Private/Models/FriendItem.h
@@ -108,6 +108,12 @@ public:
 	 */
 	virtual bool IsGameJoinable() const override;
 
+	/**
+	* Get if the user can join our game if we were to invite them
+	* @return True if we can invite them
+	*/
+	virtual bool CanInvite() const override;
+
 	/**
 	 * Get if the user is online and his game is joinable
 	 * @return The user joinable game state.
diff --git a/Engine/Source/Developer/FriendsAndChat/Private/Models/FriendRecentPlayerItem.cpp b/Engine/Source/Developer/FriendsAndChat/Private/Models/FriendRecentPlayerItem.cpp
--- a/Engine/Source/Developer/FriendsAndChat/Private/Models/FriendRecentPlayerItem.cpp
+++ b/Engine/Source/Developer/FriendsAndChat/Private/Models/FriendRecentPlayerItem.cpp
@@ -102,6 +102,11 @@ bool FFriendRecentPlayerItem::IsGameJoinable() const
 	return false;
 }
 
+bool FFriendRecentPlayerItem::CanInvite() const
+{
+	return false;
+}
+
 FString FFriendRecentPlayerItem::GetGameSessionId() const
 {
 	return FString();
diff --git a/Engine/Source/Developer/FriendsAndChat/Private/Models/FriendRecentPlayerItems.h b/Engine/Source/Developer/FriendsAndChat/Private/Models/FriendRecentPlayerItems.h
--- a/Engine/Source/Developer/FriendsAndChat/Private/Models/FriendRecentPlayerItems.h
+++ b/Engine/Source/Developer/FriendsAndChat/Private/Models/FriendRecentPlayerItems.h
@@ -128,6 +128,9 @@ public:
 	/** Is the player in a game that is joinable */
 	virtual bool IsGameJoinable() const override;
 
+	/** Get if the user can join our game if we were to invite them*/
+	virtual bool CanInvite() const override;
+
 	/** Get if the user is online and his game is joinable */
 	virtual FString GetGameSessionId() const override;
 
diff --git a/Engine/Source/Developer/FriendsAndChat/Private/Models/FriendViewModel.cpp b/Engine/Source/Developer/FriendsAndChat/Private/Models/FriendViewModel.cpp
--- a/Engine/Source/Developer/FriendsAndChat/Private/Models/FriendViewModel.cpp
+++ b/Engine/Source/Developer/FriendsAndChat/Private/Models/FriendViewModel.cpp
@@ -44,12 +44,12 @@ public:
 				{
 					if (FriendItem->IsGameJoinable())
 					{
-						if(!bFromChat || CanPerformAction(EFriendActionType::JoinGame))
+						if(CanPerformAction(EFriendActionType::JoinGame))
 						{
 							Actions.Add(EFriendActionType::JoinGame);
 						}
 					}
-					if (FFriendsAndChatManager::Get()->IsInJoinableGameSession())
+					if (FriendItem->CanInvite() && FFriendsAndChatManager::Get()->IsInJoinableGameSession())
 					{
 						Actions.Add(EFriendActionType::InviteToGame);
 					}
diff --git a/Engine/Source/Developer/FriendsAndChat/Private/Models/IFriendItem.h b/Engine/Source/Developer/FriendsAndChat/Private/Models/IFriendItem.h
--- a/Engine/Source/Developer/FriendsAndChat/Private/Models/IFriendItem.h
+++ b/Engine/Source/Developer/FriendsAndChat/Private/Models/IFriendItem.h
@@ -73,6 +73,12 @@ public:
 	 */
 	virtual bool IsGameJoinable() const = 0;
 
+	/**
+	 * Get if the user can join our game if we were to invite them
+	 * @return True if we can invite them
+	 */
+	virtual bool CanInvite() const = 0;
+
 	/**
 	 * Get game session id that this friend is currently in
 	 * @return The id of the game session
