[INTEGRATE] merge ue4-launcherdev->ue4 cl [REDACTED_PHONE]

diff --git a/Engine/Source/Developer/FriendsAndChat/Private/FriendsAndChatManager.cpp b/Engine/Source/Developer/FriendsAndChat/Private/FriendsAndChatManager.cpp
--- a/Engine/Source/Developer/FriendsAndChat/Private/FriendsAndChatManager.cpp
+++ b/Engine/Source/Developer/FriendsAndChat/Private/FriendsAndChatManager.cpp
@@ -36,17 +36,13 @@ FFriendsAndChatManager::FFriendsAndChatManager( )
 	: OnlineSub(nullptr)
 	, MessageManager(FFriendsMessageManagerFactory::Create())
 	, bJoinedGlobalChat(false)
-	, bMultiWindowChat(false)
+	, bMultiWindowChat(true)
 	, ManagerState ( EFriendsAndManagerState::OffLine )
 	, bIsInited( false )
 	, bRequiresListRefresh(false)
 	, bRequiresRecentPlayersRefresh(false)
 	, FlushChatAnalyticsCountdown(CHAT_ANALYTICS_INTERVAL)
 {
-	if (FParse::Param(FCommandLine::Get(), TEXT("ChatMultiWindow")))
-	{
-		bMultiWindowChat = true;
-	}
 }
 
 
@@ -219,9 +215,9 @@ bool FFriendsAndChatManager::IsLoggedIn()
 	return ManagerState != EFriendsAndManagerState::OffLine;
 }
 
-void FFriendsAndChatManager::SetApplicationViewModel(TSharedPtr<IFriendsApplicationViewModel> InApplicationViewModel)
+void FFriendsAndChatManager::AddApplicationViewModel(const FString ClientID, TSharedPtr<IFriendsApplicationViewModel> InApplicationViewModel)
 {
-	ApplicationViewModel = InApplicationViewModel;
+	ApplicationViewModels.Add(ClientID, InApplicationViewModel);
 }
 
 void FFriendsAndChatManager::SetUserSettings(const FFriendsAndChatSettings& UserSettings)
@@ -401,6 +397,7 @@ TSharedPtr< SWidget > FFriendsAndChatManager::GenerateFriendsListWidget( const F
 
 TSharedPtr< SWidget > FFriendsAndChatManager::GenerateChatWidget(const FFriendsAndChatStyle* InStyle, TSharedRef<IChatViewModel> ViewModel)
 {
+	bMultiWindowChat = false;
 	if(!ChatViewModel.IsValid())
 	{
 		ChatViewModel = FChatViewModelFactory::Create(MessageManager.ToSharedRef());
@@ -803,15 +800,19 @@ bool FFriendsAndChatManager::IsInJoinableGameSession() const
 	return bIsGameJoinable && IsInGameSession();
 }
 
-bool FFriendsAndChatManager::JoinGameAllowed()
+bool FFriendsAndChatManager::JoinGameAllowed(FString ClientID)
 {
 	if (AllowFriendsJoinGame().IsBound())
 	{
 		return AllowFriendsJoinGame().Execute();
 	}
-	else if (ApplicationViewModel.IsValid())
+	else
 	{
-		return ApplicationViewModel->IsAppJoinable();
+		TSharedPtr<IFriendsApplicationViewModel> FriendsApplicationViewModel = *ApplicationViewModels.Find(ClientID);
+		if (FriendsApplicationViewModel.IsValid())
+		{
+			return FriendsApplicationViewModel->IsAppJoinable();
+		}
 	}
 	return false;
 }
@@ -1714,10 +1715,11 @@ void FFriendsAndChatManager::AcceptGameInvite(const TSharedPtr<IFriendItem>& Fri
 	// notify for further processing of join game request 
 	OnFriendsJoinGame().Broadcast(*FriendItem->GetUniqueID(), FriendItem->GetGameSessionId());
 
-	if(ApplicationViewModel.IsValid())
+	TSharedPtr<IFriendsApplicationViewModel> FriendsApplicationViewModel = *ApplicationViewModels.Find(FriendItem->GetClientId());
+	if (FriendsApplicationViewModel.IsValid())
 	{
 		const FString AdditionalCommandline = TEXT("-invitesession=") + FriendItem->GetGameSessionId() + TEXT(" -invitefrom=") + FriendItem->GetUniqueID()->ToString();
-		ApplicationViewModel->LaunchFriendApp(AdditionalCommandline);
+		FriendsApplicationViewModel->LaunchFriendApp(AdditionalCommandline);
 	}
 
 	Analytics.RecordGameInvite(*FriendItem->GetUniqueID(), TEXT("Social.GameInvite.Accept"));
diff --git a/Engine/Source/Developer/FriendsAndChat/Private/FriendsAndChatManager.h b/Engine/Source/Developer/FriendsAndChat/Private/FriendsAndChatManager.h
--- a/Engine/Source/Developer/FriendsAndChat/Private/FriendsAndChatManager.h
+++ b/Engine/Source/Developer/FriendsAndChat/Private/FriendsAndChatManager.h
@@ -125,7 +125,7 @@ public:
 	virtual void Logout() override;
 	virtual void Login() override;
 	virtual bool IsLoggedIn() override;
-	virtual void SetApplicationViewModel(TSharedPtr<IFriendsApplicationViewModel> ApplicationViewModel) override;
+	virtual void AddApplicationViewModel(const FString ClientID, TSharedPtr<IFriendsApplicationViewModel> ApplicationViewModel) override;
 	virtual void CreateFriendsListWindow(const FFriendsAndChatStyle* InStyle) override;
 	virtual void CreateChatWindow(const struct FFriendsAndChatStyle* InStyle, EChatMessageType::Type ChatType, TSharedPtr<IFriendItem> FriendItem) override;
 	virtual void SetUserSettings(const FFriendsAndChatSettings& UserSettings) override;
@@ -167,9 +167,10 @@ public:
 	bool IsInJoinableGameSession() const;
 
 	/**
+	 * @param ClientID ID of the Game we want to join
 	 * @return true if joining a game is allowed
 	 */
-	bool JoinGameAllowed();
+	bool JoinGameAllowed(FString ClientID);
 
 	/**
 	 * @return true if in the launcher
@@ -816,8 +817,8 @@ private:
 	TArray<WhisperChat> WhisperChatWindows;
 
 
-	// Holds the application view model - used for launching and querying
-	TSharedPtr<IFriendsApplicationViewModel> ApplicationViewModel;
+	// Holds the application view models - used for launching and querying the games
+	TMap<FString, TSharedPtr<IFriendsApplicationViewModel>> ApplicationViewModels;
 	// Holds the style used to create the Friends List widget
 	FFriendsAndChatStyle Style;
 	// Holds if the Friends list is inited
diff --git a/Engine/Source/Developer/FriendsAndChat/Private/Models/FriendViewModel.cpp b/Engine/Source/Developer/FriendsAndChat/Private/Models/FriendViewModel.cpp
--- a/Engine/Source/Developer/FriendsAndChat/Private/Models/FriendViewModel.cpp
+++ b/Engine/Source/Developer/FriendsAndChat/Private/Models/FriendViewModel.cpp
@@ -140,7 +140,7 @@ public:
 		{
 			case EFriendActionType::JoinGame:
 			{
-				return FFriendsAndChatManager::Get()->JoinGameAllowed();
+				return FFriendsAndChatManager::Get()->JoinGameAllowed(GetClientId());
 			}
 			case EFriendActionType::AcceptFriendRequest:
 			case EFriendActionType::RemoveFriend:
diff --git a/Engine/Source/Developer/FriendsAndChat/Private/UI/Widgets/SFriendsContainer.cpp b/Engine/Source/Developer/FriendsAndChat/Private/UI/Widgets/SFriendsContainer.cpp
--- a/Engine/Source/Developer/FriendsAndChat/Private/UI/Widgets/SFriendsContainer.cpp
+++ b/Engine/Source/Developer/FriendsAndChat/Private/UI/Widgets/SFriendsContainer.cpp
@@ -276,6 +276,11 @@ private:
 		return FReply::Handled();
 	}
 
+	virtual FReply OnMouseButtonDown(const FGeometry& MyGeometry, const FPointerEvent& MouseEvent) override
+	{
+		return FReply::Handled();
+	}
+
 private:
 
 	// Holds the Friends List view model
diff --git a/Engine/Source/Developer/FriendsAndChat/Public/Interfaces/IFriendsAndChatManager.h b/Engine/Source/Developer/FriendsAndChat/Public/Interfaces/IFriendsAndChatManager.h
--- a/Engine/Source/Developer/FriendsAndChat/Public/Interfaces/IFriendsAndChatManager.h
+++ b/Engine/Source/Developer/FriendsAndChat/Public/Interfaces/IFriendsAndChatManager.h
@@ -92,9 +92,10 @@ public:
 
 	/** 
 	 * Set the application view model to query and perform actions on.
+	 * @param ClientID The ID of the application
 	 * @param ApplicationViewModel The view model.
 	 */
-	virtual void SetApplicationViewModel(TSharedPtr<IFriendsApplicationViewModel> ApplicationViewModel) = 0;
+	virtual void AddApplicationViewModel(const FString ClientID, TSharedPtr<IFriendsApplicationViewModel> ApplicationViewModel) = 0;
 
 	DECLARE_EVENT_OneParam(IFriendsAndChatManager, FOnFriendsNotificationEvent, const bool /*Show or Clear */)
 	virtual FOnFriendsNotificationEvent& OnFriendsNotification() = 0;
