[INTEGRATE] merge ue4-UT->ue4  FriendsAndChat Module

#proj Portal

diff --git a/Engine/Source/Developer/FriendsAndChat/Private/FriendsAndChatManager.cpp b/Engine/Source/Developer/FriendsAndChat/Private/FriendsAndChatManager.cpp
--- a/Engine/Source/Developer/FriendsAndChat/Private/FriendsAndChatManager.cpp
+++ b/Engine/Source/Developer/FriendsAndChat/Private/FriendsAndChatManager.cpp
@@ -16,7 +16,6 @@
 #include "FriendGameInviteItem.h"
 #include "ChatDisplayOptionsViewModel.h"
 #include "SFriendUserHeader.h"
-#include "OnlineIdentityMcp.h"
 
 #define LOCTEXT_NAMESPACE "FriendsAndChatManager"
 
@@ -39,8 +38,8 @@ FFriendsAndChatManager::FFriendsAndChatManager( )
 	: OnlineSub(nullptr)
 	, MessageManager(FFriendsMessageManagerFactory::Create())
 	, bJoinedGlobalChat(false)
-	, bMultiWindowChat(true)
 	, bHasGlobalChatPermission(false)
+	, ChatWindowMode(EChatWindowMode::MultiWindow)
 	, ManagerState ( EFriendsAndManagerState::OffLine )
 	, bIsInited( false )
 	, bRequiresListRefresh(false)
@@ -464,7 +463,7 @@ TSharedPtr< SWidget > FFriendsAndChatManager::GenerateStatusWidget( const FFrien
 
 TSharedPtr< SWidget > FFriendsAndChatManager::GenerateChatWidget(const FFriendsAndChatStyle* InStyle, TSharedRef<IChatViewModel> ViewModel, TAttribute<FText> ActivationHintDelegate)
 {
-	bMultiWindowChat = false;
+	ChatWindowMode = EChatWindowMode::Widget;
 	if(!ChatViewModel.IsValid())
 	{
 		ChatViewModel = FChatViewModelFactory::Create(MessageManager.ToSharedRef());
@@ -534,7 +533,7 @@ void FFriendsAndChatManager::CreateChatWindow(const struct FFriendsAndChatStyle*
 
 	// Look up if window has already been created
 	TSharedPtr<SWindow> ChatWindow;
-	if (ChatType == EChatMessageType::Whisper && FriendItem.IsValid() && bMultiWindowChat)
+	if (ChatType == EChatMessageType::Whisper && FriendItem.IsValid() && ChatWindowMode == EChatWindowMode::MultiWindow)
 	{
 		for (const auto& WhisperChatWindow : WhisperChatWindows)
 		{
@@ -557,7 +556,7 @@ void FFriendsAndChatManager::CreateChatWindow(const struct FFriendsAndChatStyle*
 		// Choose Window Name
 		FText WindowName = LOCTEXT("FriendsAndChatManager_ChatTitle", "Chat Window");
 
-		if (bMultiWindowChat)
+		if (ChatWindowMode != EChatWindowMode::MultiWindow)
 		{
 			if (ChatType == EChatMessageType::Global)
 			{
@@ -599,14 +598,14 @@ void FFriendsAndChatManager::CreateChatWindow(const struct FFriendsAndChatStyle*
 		ChatWindow = FSlateApplication::Get().AddWindow(ChatWindow.ToSharedRef());
 
 		// Store window ptr
-		if (ChatType == EChatMessageType::Whisper && FriendItem.IsValid() && bMultiWindowChat)
+		if (ChatType == EChatMessageType::Whisper && FriendItem.IsValid() && ChatWindowMode == EChatWindowMode::MultiWindow)
 		{
 			WhisperChat NewWhisperChat;
 			NewWhisperChat.FriendItem = FriendItem;
 			NewWhisperChat.ChatWindow = ChatWindow;
 			WhisperChatWindows.Add(NewWhisperChat);
 		}
-		else if (ChatType == EChatMessageType::Global || !bMultiWindowChat)
+		else if (ChatType == EChatMessageType::Global || ChatWindowMode != EChatWindowMode::MultiWindow)
 		{
 			GlobalChatWindow = ChatWindow;
 		}
@@ -621,7 +620,7 @@ void FFriendsAndChatManager::CreateChatWindow(const struct FFriendsAndChatStyle*
 		{
 		ChatWindow->Restore();
 		}
-		if (!bMultiWindowChat)
+		if (ChatWindowMode != EChatWindowMode::MultiWindow)
 		{
 			SetChatWindowContents(ChatWindow, nullptr, Style);
 		}		
@@ -641,7 +640,15 @@ void FFriendsAndChatManager::CreateChatWindow(const struct FFriendsAndChatStyle*
 
 void FFriendsAndChatManager::SetChatFriend( TSharedPtr< IFriendItem > FriendItem )
 {
-	CreateChatWindow(&Style, EChatMessageType::Whisper, FriendItem, true);
+	if (ChatWindowMode != EChatWindowMode::Widget)
+	{
+		CreateChatWindow(&Style, EChatMessageType::Whisper, FriendItem, true);
+	}
+
+	if (ChatWindowMode != EChatWindowMode::MultiWindow)
+	{
+		OnChatFriendSelected().Broadcast(FriendItem);
+	}
 }
 
 void FFriendsAndChatManager::OpenGlobalChat()
@@ -778,7 +785,7 @@ private:
 void FFriendsAndChatManager::SetChatWindowContents(TSharedPtr<SWindow> Window, TSharedPtr< IFriendItem > FriendItem, const FFriendsAndChatStyle& FriendStyle)
 {
 	TSharedPtr<SWindowTitleBar> TitleBar;	
-	if(!ChatViewModel.IsValid() || bMultiWindowChat)
+	if(!ChatViewModel.IsValid() || ChatWindowMode == EChatWindowMode::MultiWindow)
 	{
 		ChatViewModel = FChatViewModelFactory::Create(MessageManager.ToSharedRef());
 	}
@@ -814,7 +821,7 @@ void FFriendsAndChatManager::SetChatWindowContents(TSharedPtr<SWindow> Window, T
 	}
 
 	// Lock the channel if not using unified chat
-	if (bMultiWindowChat)
+	if (ChatWindowMode == EChatWindowMode::MultiWindow)
 	{
 		DisplayViewModel->GetChatViewModel()->LockChatChannel(true);
 
@@ -1094,12 +1101,12 @@ bool FFriendsAndChatManager::JoinGameAllowed(FString ClientID)
 		if (FriendsApplicationViewModelRawPtr != nullptr)
 		{
 			TSharedPtr<IFriendsApplicationViewModel> FriendsApplicationViewModel = *FriendsApplicationViewModelRawPtr;
-		if (FriendsApplicationViewModel.IsValid())
-		{
+			if (FriendsApplicationViewModel.IsValid())
+			{
 				return FriendsApplicationViewModel->IsAppJoinable();
+			}
 		}
 	}
-	}
 	return false;
 }
 
@@ -1955,30 +1962,40 @@ void FFriendsAndChatManager::ProcessReceivedGameInvites()
 		{
 			const FReceivedGameInvite& Invite = ReceivedGameInvites[Idx];
 
-			TSharedPtr<FOnlineUser> UserInfo;
-			TSharedPtr<IFriendItem> Friend = FindUser(*Invite.FromId);
-			TSharedPtr<FOnlineFriend> OnlineFriend;
-			if (Friend.IsValid())
+			if (!Invite.InviteResult->Session.SessionInfo.IsValid() ||
+				Invite.InviteResult->Session.SessionInfo->GetSessionId().ToString() == GetGameSessionId()->ToString())
 			{
-				UserInfo = Friend->GetOnlineUser();
-				OnlineFriend = Friend->GetOnlineFriend();
-			}
-			if (!UserInfo.IsValid())
-			{
-				UserInfo = OnlineSub->GetUserInterface()->GetUserInfo(0, *Invite.FromId);
+				// remove invites if user is already in the game session
+				ReceivedGameInvites.RemoveAt(Idx--);
 			}
-			if (UserInfo.IsValid() && OnlineFriend.IsValid())
+			else
 			{
-				TSharedPtr<FFriendGameInviteItem> GameInvite = MakeShareable(
-					new FFriendGameInviteItem(UserInfo.ToSharedRef(), Invite.InviteResult, Invite.ClientId, OnlineFriend.ToSharedRef())
-					);
+				// add to list of pending invites to accept
+				TSharedPtr<FOnlineUser> UserInfo;
+				TSharedPtr<IFriendItem> Friend = FindUser(*Invite.FromId);
+				TSharedPtr<FOnlineFriend> OnlineFriend;
+				if (Friend.IsValid())
+				{
+					UserInfo = Friend->GetOnlineUser();
+					OnlineFriend = Friend->GetOnlineFriend();
+				}
+				if (!UserInfo.IsValid())
+				{
+					UserInfo = OnlineSub->GetUserInterface()->GetUserInfo(0, *Invite.FromId);
+				}
+				if (UserInfo.IsValid() && OnlineFriend.IsValid())
+				{
+					TSharedPtr<FFriendGameInviteItem> GameInvite = MakeShareable(
+						new FFriendGameInviteItem(UserInfo.ToSharedRef(), Invite.InviteResult, Invite.ClientId, OnlineFriend.ToSharedRef())
+						);
 
-				PendingGameInvitesList.Add(Invite.FromId->ToString(), GameInvite);
+					PendingGameInvitesList.Add(Invite.FromId->ToString(), GameInvite);
 
-				OnGameInvitesUpdated().Broadcast();
-				SendGameInviteNotification(GameInvite);
+					OnGameInvitesUpdated().Broadcast();
+					SendGameInviteNotification(GameInvite);
 
-				ReceivedGameInvites.RemoveAt(Idx--);
+					ReceivedGameInvites.RemoveAt(Idx--);
+				}
 			}
 		}
 	}
@@ -2064,18 +2081,13 @@ void FFriendsAndChatManager::AcceptGameInvite(const TSharedPtr<IFriendItem>& Fri
 	// notify for further processing of join game request 
 	OnFriendsJoinGame().Broadcast(*FriendItem->GetUniqueID(), *FriendItem->GetGameSessionId());
 
-	TSharedPtr<IFriendsApplicationViewModel>* FriendsApplicationViewModelRawPtr = ApplicationViewModels.Find(FriendItem->GetClientId());
-	if (FriendsApplicationViewModelRawPtr != nullptr)
-	{
-		TSharedPtr<IFriendsApplicationViewModel> FriendsApplicationViewModel = *FriendsApplicationViewModelRawPtr;
+	TSharedPtr<IFriendsApplicationViewModel> FriendsApplicationViewModel = *ApplicationViewModels.Find(FriendItem->GetClientId());
 	if (FriendsApplicationViewModel.IsValid())
 	{
-			const FString AdditionalCommandline = TEXT("-invitesession=") + FriendItem->GetGameSessionId()->ToString() + TEXT(" -invitefrom=") + FriendItem->GetUniqueID()->ToString();
+		const FString AdditionalCommandline = TEXT("-invitesession=") + FriendItem->GetGameSessionId()->ToString() + TEXT(" -invitefrom=") + FriendItem->GetUniqueID()->ToString();
 		FriendsApplicationViewModel->LaunchFriendApp(AdditionalCommandline);
 	}
 
-	}
-
 	Analytics.RecordGameInvite(*FriendItem->GetUniqueID(), TEXT("Social.GameInvite.Accept"));
 }
 
diff --git a/Engine/Source/Developer/FriendsAndChat/Private/FriendsAndChatManager.h b/Engine/Source/Developer/FriendsAndChat/Private/FriendsAndChatManager.h
--- a/Engine/Source/Developer/FriendsAndChat/Private/FriendsAndChatManager.h
+++ b/Engine/Source/Developer/FriendsAndChat/Private/FriendsAndChatManager.h
@@ -806,11 +806,19 @@ private:
 	TSharedPtr<class FChatViewModel> ChatViewModel;
 	// Joined Global Chat
 	bool bJoinedGlobalChat;
-	// Use one window for each chat
-	bool bMultiWindowChat;
 	// Has Permission to use global chat
 	bool bHasGlobalChatPermission;
 
+	enum class EChatWindowMode : uint8
+	{
+		Widget,
+		SingleWindow,
+		MultiWindow
+	};
+
+	// Use chat window, one window for each chat, non-windowed widgets?
+	EChatWindowMode ChatWindowMode;
+
 	/* Manger state
 	*****************************************************************************/
 
diff --git a/Engine/Source/Developer/FriendsAndChat/Private/Models/ChatDisplayOptionsViewModel.cpp b/Engine/Source/Developer/FriendsAndChat/Private/Models/ChatDisplayOptionsViewModel.cpp
--- a/Engine/Source/Developer/FriendsAndChat/Private/Models/ChatDisplayOptionsViewModel.cpp
+++ b/Engine/Source/Developer/FriendsAndChat/Private/Models/ChatDisplayOptionsViewModel.cpp
@@ -63,6 +63,11 @@ public:
 		return ChatViewModel->IsGlobalChatEnabled();
 	}
 
+	virtual float GetChatListFadeValue() const override
+	{
+		return GetTimeTransparency();
+	}
+
 // End IChatViewModel interface
 
 // Begin ChatDisplayOptionsViewModel interface
diff --git a/Engine/Source/Developer/FriendsAndChat/Private/UI/Widgets/SChatWindow.cpp b/Engine/Source/Developer/FriendsAndChat/Private/UI/Widgets/SChatWindow.cpp
--- a/Engine/Source/Developer/FriendsAndChat/Private/UI/Widgets/SChatWindow.cpp
+++ b/Engine/Source/Developer/FriendsAndChat/Private/UI/Widgets/SChatWindow.cpp
@@ -101,6 +101,7 @@ public:
 					[
 						SAssignNew(ActionMenu, SMenuAnchor)
 						.Placement(MenuMethod != EPopupMethod::UseCurrentWindow ? EMenuPlacement::MenuPlacement_AboveAnchor : EMenuPlacement::MenuPlacement_BelowAnchor)  	// MenuPlacement_BelowAnchor is a workaround until a workarea bug is fixed in SlateApplication menu placement code
+						.IsEnabled(this, &SChatWindowImpl::HasMenuOptions)
 						.Method(EPopupMethod::UseCurrentWindow)
 						.OnGetMenuContent(this, &SChatWindowImpl::GetMenuContent)
 						.Visibility(this, &SChatWindowImpl::GetChatChannelVisibility)
@@ -434,14 +435,23 @@ public:
 			{
 				if (FadeCurve.GetLerp() == 0.0f)
 				{
-					FadeAnimation.Play(this->AsShared());
+					FadeAnimation.Play();
 				}
 			}
 			else
 			{
 				if (FadeCurve.GetLerp() == 1.0f)
 				{
-					FadeAnimation.PlayReverse(this->AsShared());
+					FadeAnimation.PlayReverse();
+				}
+
+				if (FadeAnimation.IsPlaying())
+				{
+				    TArray<TSharedRef<FChatItemViewModel>>& FilteredChatList = SharedChatViewModel->GetMessages();
+				    if (FilteredChatList.Num() > 0)
+				    {
+					    ChatList->RequestScrollIntoView(FilteredChatList.Last());
+				    }
 				}
 			}
 		}
@@ -648,6 +658,28 @@ private:
 		return Contents;
 	}
 
+	bool HasMenuOptions() const
+	{
+		if (SharedChatViewModel->IsGlobalChatEnabled())
+		{
+			return true;
+		}
+
+		if (SharedChatViewModel->GetRecentOptions().Num() > 0)
+		{
+			return true;
+		}
+
+		TArray<EChatMessageType::Type> ChatMessageOptions;
+		DisplayViewModel->EnumerateChatChannelOptionsList(ChatMessageOptions);
+		if (ChatMessageOptions.Num() > 0)
+		{
+			return true;
+		}
+
+		return false;
+	}
+
 	TSharedRef<SWidget> GetFriendActionMenu()
 	{
 		TSharedPtr<SVerticalBox> FriendActionBox;
diff --git a/Engine/Source/Developer/FriendsAndChat/Public/Models/IChatViewModel.h b/Engine/Source/Developer/FriendsAndChat/Public/Models/IChatViewModel.h
--- a/Engine/Source/Developer/FriendsAndChat/Public/Models/IChatViewModel.h
+++ b/Engine/Source/Developer/FriendsAndChat/Public/Models/IChatViewModel.h
@@ -16,6 +16,7 @@ public:
 	virtual FSlateColor GetFontOverrideColor() const = 0;
 	virtual void SetInGameUI(bool bInGame) = 0;
 	virtual void EnableGlobalChat(bool bEnable) = 0;
+	virtual float GetChatListFadeValue() const = 0;
 	virtual bool IsGlobalChatEnabled() const = 0;
 
 	DECLARE_EVENT(IChatViewModel, FOnFriendsChatMessageCommitted)
