Fixed up our combo buttons to use the new way of creating menu content

diff --git a/Engine/Source/Developer/FriendsAndChat/Private/UI/Widgets/SFriendsAndChatCombo.cpp b/Engine/Source/Developer/FriendsAndChat/Private/UI/Widgets/SFriendsAndChatCombo.cpp
--- a/Engine/Source/Developer/FriendsAndChat/Private/UI/Widgets/SFriendsAndChatCombo.cpp
+++ b/Engine/Source/Developer/FriendsAndChat/Private/UI/Widgets/SFriendsAndChatCombo.cpp
@@ -99,7 +99,6 @@ public:
 		MenuBorderPadding = ComboButtonStyle->MenuBorderPadding;
 
 		OnComboBoxOpened = InArgs._OnDropdownOpened;
-		OnGetMenuContent = FOnGetContent::CreateSP(this, &SFriendsAndChatComboButton::GetMenuContent);
 
 		// Purposefully not calling SComboButton's constructor as we want to override more of the visuals
 
@@ -109,6 +108,7 @@ public:
 		SMenuAnchor::Construct(SMenuAnchor::FArguments()
 			.Placement(InArgs._Placement)
 			.Method(EPopupMethod::UseCurrentWindow)
+			.OnGetMenuContent(this, &SFriendsAndChatComboButton::GetMenuContent)
 			[
 				SAssignNew(DropdownButton, SCustomDropdownButton)
 				.ButtonStyle(&ComboButtonStyle->ButtonStyle)
@@ -183,12 +183,6 @@ public:
 
 protected:
 
-	/** Overridden to pass button click to SComboBox */
-	virtual FReply OnButtonClicked() override
-	{
-		return SComboButton::OnButtonClicked();
-	}
-
 	/** Returns pointer to dropdown button */
 	TSharedPtr<SButton> GetDropdownButton()
 	{
