[INTEGRATE] Merging using UE4-To-UE4-LauncherDev, integrate all from LauncherDev up to cl 2499790, more party changes portal related bug fixes

diff --git a/Engine/Source/Developer/FriendsAndChat/Private/ClanRepository.cpp b/Engine/Source/Developer/FriendsAndChat/Private/ClanRepository.cpp
--- a/Engine/Source/Developer/FriendsAndChat/Private/ClanRepository.cpp
+++ b/Engine/Source/Developer/FriendsAndChat/Private/ClanRepository.cpp
@@ -186,6 +186,16 @@ public:
 		return MemberList;
 	}
 
+	virtual FText GetClanBrushName() const override
+	{
+		return FText::FromString("TileMapEditor.LayerEyeOpened");
+	}
+
+	virtual int32 GetClanMembersCount() const override
+	{
+		return MemberList.Num();
+	}
+
 	DECLARE_DERIVED_EVENT(FClanInfo, IClanInfo::FChangedEvent, FChangedEvent)
 	virtual IClanInfo::FChangedEvent& OnChanged() override { return ChangedEvent; }
 
diff --git a/Engine/Source/Developer/FriendsAndChat/Private/FriendsAndChatStyle.cpp b/Engine/Source/Developer/FriendsAndChat/Private/FriendsAndChatStyle.cpp
--- a/Engine/Source/Developer/FriendsAndChat/Private/FriendsAndChatStyle.cpp
+++ b/Engine/Source/Developer/FriendsAndChat/Private/FriendsAndChatStyle.cpp
@@ -543,6 +543,17 @@ FFriendsAndChatStyle& FFriendsAndChatStyle::SetChatWindowStyle(const EChatChanne
 	return *this;
 }
 
+FFriendsAndChatStyle& FFriendsAndChatStyle::SetClanDetailsBrush(const FSlateBrush& Brush)
+{
+	ClanDetailsBrush = Brush;
+	return *this;
+}
+
+FFriendsAndChatStyle& FFriendsAndChatStyle::SetClanMembersBrush(const FSlateBrush& Brush)
+{
+	ClanMembersBrush = Brush;
+	return *this;
+}
 
 const FFriendsAndChatStyle& FFriendsAndChatStyle::GetDefault()
 {
diff --git a/Engine/Source/Developer/FriendsAndChat/Private/Models/ClanInfo.h b/Engine/Source/Developer/FriendsAndChat/Private/Models/ClanInfo.h
--- a/Engine/Source/Developer/FriendsAndChat/Private/Models/ClanInfo.h
+++ b/Engine/Source/Developer/FriendsAndChat/Private/Models/ClanInfo.h
@@ -11,6 +11,8 @@ public:
 	virtual FText GetDescription() const = 0;
 	virtual FText GetLongDescription() const = 0;
 	virtual TArray<TSharedPtr<class IFriendItem>> GetMemberList() const = 0;
+	virtual FText GetClanBrushName() const = 0;
+	virtual int32 GetClanMembersCount() const = 0;
 
 	DECLARE_EVENT(IClanInfo, FChangedEvent)
 	virtual FChangedEvent& OnChanged() = 0;
diff --git a/Engine/Source/Developer/FriendsAndChat/Private/Models/ClanInfoViewModel.cpp b/Engine/Source/Developer/FriendsAndChat/Private/Models/ClanInfoViewModel.cpp
--- a/Engine/Source/Developer/FriendsAndChat/Private/Models/ClanInfoViewModel.cpp
+++ b/Engine/Source/Developer/FriendsAndChat/Private/Models/ClanInfoViewModel.cpp
@@ -31,6 +31,11 @@ public:
 		return FText::AsNumber(ClanInfo->GetMemberList().Num());
 	}
 
+	virtual FText GetClanBrushName() const override
+	{
+		return ClanInfo->GetClanBrushName();
+	}
+
 private:
 	void Initialize()
 	{
diff --git a/Engine/Source/Developer/FriendsAndChat/Private/Models/ClanInfoViewModel.h b/Engine/Source/Developer/FriendsAndChat/Private/Models/ClanInfoViewModel.h
--- a/Engine/Source/Developer/FriendsAndChat/Private/Models/ClanInfoViewModel.h
+++ b/Engine/Source/Developer/FriendsAndChat/Private/Models/ClanInfoViewModel.h
@@ -12,6 +12,7 @@ public:
 	virtual int32 GetMemberCount() = 0;
 	virtual TSharedRef<class FFriendListViewModel > GetFriendListViewModel() = 0;
 	virtual FText GetListCountText() const = 0;
+	virtual FText GetClanBrushName() const = 0;
 };
 
 /**
diff --git a/Engine/Source/Developer/FriendsAndChat/Private/Models/ClanListViewModel.h b/Engine/Source/Developer/FriendsAndChat/Private/Models/ClanListViewModel.h
--- a/Engine/Source/Developer/FriendsAndChat/Private/Models/ClanListViewModel.h
+++ b/Engine/Source/Developer/FriendsAndChat/Private/Models/ClanListViewModel.h
@@ -5,7 +5,7 @@
 // Enum holding the display list selection
 namespace EClanDisplayLists
 {
-	enum Type
+	enum Type : uint8
 	{
 		DefaultDisplay,
 		ClanRequestsDisplay,			// Clan request display
diff --git a/Engine/Source/Developer/FriendsAndChat/Private/UI/Widgets/SClanHome.cpp b/Engine/Source/Developer/FriendsAndChat/Private/UI/Widgets/SClanHome.cpp
--- a/Engine/Source/Developer/FriendsAndChat/Private/UI/Widgets/SClanHome.cpp
+++ b/Engine/Source/Developer/FriendsAndChat/Private/UI/Widgets/SClanHome.cpp
@@ -3,6 +3,7 @@
 #include "FriendsAndChatPrivatePCH.h"
 #include "SClanHome.h"
 #include "SClanList.h"
+#include "SClanListContainer.h"
 #include "ClanViewModel.h"
 #include "ClanListViewModel.h"
 
@@ -24,13 +25,20 @@ public:
 		[
 			SNew(SVerticalBox)
 			+ SVerticalBox::Slot()
+			.Padding(10)
 			[
-				SNew(STextBlock)
-				.Text(FText::FromString("Clan Home"))
+				SNew(SClanListContainer, InViewModel)
+				.FriendStyle(&FriendStyle)
+				.Visibility(EVisibility::Visible)
+				.ClanListType(EClanDisplayLists::DefaultDisplay)
 			]
 			+ SVerticalBox::Slot()
+			.Padding(10)
 			[
-				SNew(SClanList, FClanListViewModelFactory::Create(InViewModel, EClanDisplayLists::DefaultDisplay)).FriendStyle(&FriendStyle)
+				SNew(SClanListContainer, InViewModel)
+				.FriendStyle(&FriendStyle)
+				.Visibility(EVisibility::Collapsed)
+				.ClanListType(EClanDisplayLists::ClanRequestsDisplay)
 			]
 		]);
 	}
diff --git a/Engine/Source/Developer/FriendsAndChat/Private/UI/Widgets/SClanItem.cpp b/Engine/Source/Developer/FriendsAndChat/Private/UI/Widgets/SClanItem.cpp
--- a/Engine/Source/Developer/FriendsAndChat/Private/UI/Widgets/SClanItem.cpp
+++ b/Engine/Source/Developer/FriendsAndChat/Private/UI/Widgets/SClanItem.cpp
@@ -16,23 +16,142 @@ public:
 		FriendStyle = *InArgs._FriendStyle;
 		this->ClanInfo = ClanInfo;
 
+		FFriendsAndChatComboButtonStyle ActionButtonStyle;
+		ActionButtonStyle.ComboButtonStyle = &FriendStyle.ActionComboButtonStyle;
+		ActionButtonStyle.ButtonTextStyle = &FriendStyle.ActionComboButtonTextStyle;
+
 		SUserWidget::Construct(SUserWidget::FArguments()
 		[
-			SNew(STextBlock)
-			.Font(FriendStyle.FriendsFontStyleBold)
-			.ColorAndOpacity(FriendStyle.DefaultFontColor)
-			.Text(ClanInfo->GetTitle())
+			SNew(SBorder)
+			.Padding(FMargin(0))
+			.BorderBackgroundColor(FLinearColor::Transparent)
+			[
+				SNew(SVerticalBox)
+				+SVerticalBox::Slot()
+				.AutoHeight()
+				[
+					SNew(SHorizontalBox)
+					+ SHorizontalBox::Slot()
+					.Padding(10, 10)
+					.AutoWidth()
+					.VAlign(VAlign_Top)
+					.HAlign(HAlign_Left)
+					[
+						SNew(SOverlay)
+						+SOverlay::Slot()
+						[
+							// Clan Image
+							SNew(SImage)
+							.Image(&FriendStyle.UTImageBrush)
+						]
+						+SOverlay::Slot()
+						.VAlign(VAlign_Top)
+						.HAlign(HAlign_Right)
+						[
+							// Primary Clan Indicator
+							SNew(SImage)
+							.Image(&FriendStyle.AwayBrush )
+						]
+					]
+					+ SHorizontalBox::Slot()
+					.Padding(15, 0)
+					.VAlign(VAlign_Center)
+					.HAlign(HAlign_Right)
+					[
+						SAssignNew(ActionMenuButton, SFriendsAndChatCombo)
+						.FriendStyle(&FriendStyle)
+						.ButtonStyleOverride(&ActionButtonStyle)
+						.ButtonText(FText::GetEmpty())
+						.bShowIcon(false)
+						.DropdownItems(this, &SClanItemImpl::GetActionItems)
+						.bSetButtonTextToSelectedItem(false)
+						.bAutoCloseWhenClicked(true)
+						.ButtonSize(FriendStyle.ActionComboButtonSize)
+						.Placement(MenuPlacement_ComboBoxRight)
+						.OnDropdownItemClicked(this, &SClanItemImpl::HandleItemClicked)
+						.OnDropdownOpened(this, &SClanItemImpl::HandleActionMenuOpened)
+						.Visibility(this, &SClanItemImpl::ActionMenuButtonVisibility)
+						.Cursor(EMouseCursor::Hand)
+					]
+					+ SHorizontalBox::Slot()
+					.AutoWidth()
+					.HAlign(HAlign_Right)
+					.Padding(0, 10)
+					[
+						SNew(SButton)
+						.ButtonStyle(&FriendStyle.GlobalChatButtonStyle)
+						[
+							SNew(SHorizontalBox)
+							+SHorizontalBox::Slot()
+							.AutoWidth()
+							[
+								SNew(SImage)
+								.Image(&FriendStyle.ClanMembersBrush)
+							]
+							+SHorizontalBox::Slot()
+							.VAlign(VAlign_Center)
+							.HAlign(HAlign_Center)
+							.Padding(10)
+							[
+								SNew(STextBlock)
+								.Font(FriendStyle.FriendsFontStyleBold)
+								.ColorAndOpacity(FriendStyle.DefaultFontColor)
+								.Text(this, &SClanItemImpl::GetClanMemberCountText)
+							]
+							+ SHorizontalBox::Slot()
+							.AutoWidth()
+							[
+								SNew(SImage)
+								.Image(&FriendStyle.ClanDetailsBrush)
+							]
+						]
+					]
+				]
+				+ SVerticalBox::Slot()
+				.AutoHeight()
+				.Padding(10, 0)
+				[
+					SNew(STextBlock)
+					.Font(FriendStyle.FriendsFontStyleBold)
+					.ColorAndOpacity(FriendStyle.DefaultFontColor)
+					.Text(ClanInfo->GetTitle())
+				]
+			]
 		]);
 	}
 
 private:
 
+	SFriendsAndChatCombo::FItemsArray GetActionItems() const
+	{
+		SFriendsAndChatCombo::FItemsArray ActionItems;
+		// Need to Add some actions
+		return ActionItems;
+	}
+
+	void HandleItemClicked(FName ItemTag)
+	{
+	}
+
+	void HandleActionMenuOpened() const
+	{
+	}
+
+	EVisibility ActionMenuButtonVisibility() const
+	{
+		return EVisibility::Visible;
+	}
+
 	virtual bool SupportsKeyboardFocus() const override
 	{
 		return true;
 	}
 
-private:
+	FText GetClanMemberCountText() const
+	{
+		return FText::AsNumber(ClanInfo->GetClanMembersCount());
+	}
+
 
 	TSharedPtr<IClanInfo> ClanInfo;
 
@@ -42,6 +161,10 @@ private:
 	TSharedPtr<SWidget> MenuContent;
 
 	EPopupMethod MenuMethod;
+
+	TSharedPtr<SFriendsAndChatCombo> ActionMenuButton;
+
+	/** Cached actual dropdown button */
 };
 
 TSharedRef<SClanItem> SClanItem::New()
diff --git a/Engine/Source/Developer/FriendsAndChat/Private/UI/Widgets/SClanListContainer.cpp b/Engine/Source/Developer/FriendsAndChat/Private/UI/Widgets/SClanListContainer.cpp
--- a/Engine/Source/Developer/FriendsAndChat/Private/UI/Widgets/SClanListContainer.cpp
+++ b/Engine/Source/Developer/FriendsAndChat/Private/UI/Widgets/SClanListContainer.cpp
@@ -14,19 +14,18 @@ class SClanListContainerImpl : public SClanListContainer
 {
 public:
 
-	void Construct(const FArguments& InArgs, const TSharedRef<FClanListViewModel>& InViewModel)
+	void Construct(const FArguments& InArgs, const TSharedRef<FClanViewModel>& InViewModel)
 	{
 		FriendStyle = *InArgs._FriendStyle;
-		FriendsVisibility = EVisibility::Visible;
+		ClanListType = InArgs._ClanListType;
+		ClansVisibility = EVisibility::Visible;
 		ViewModel = InViewModel;
 
-		FClanListViewModel* ViewModelPtr = ViewModel.Get();
 		MenuMethod = InArgs._Method;
 
 		SUserWidget::Construct(SUserWidget::FArguments()
 		[
 			SNew(SVerticalBox)
-			.Visibility(ViewModelPtr, &FClanListViewModel::GetListVisibility)
 			+ SVerticalBox::Slot()
 			.VAlign(VAlign_Center)
 			.AutoHeight()
@@ -34,14 +33,56 @@ public:
 				SNew(SBorder)
 				.Padding(FriendStyle.BorderPadding)
 				.BorderImage(&FriendStyle.FriendListHeader)
+				[
+					SNew(SButton)
+					.OnClicked(this, &SClanListContainerImpl::HandleShowClansClicked)
+					.ButtonStyle(FCoreStyle::Get(), "NoBorder")
+					.ContentPadding(FMargin(7.0f, 2.0f, 4.0f, 2.0f))
+					.Cursor(EMouseCursor::Hand)
+					[
+						SNew(SHorizontalBox)
+						+ SHorizontalBox::Slot()
+						.HAlign(HAlign_Left)
+						.VAlign(VAlign_Center)
+						.AutoWidth()
+						[
+							SNew(SButton)
+							.OnClicked(this, &SClanListContainerImpl::HandleShowClansClicked)
+							.ButtonStyle(&FriendStyle.FriendListOpenButtonStyle)
+							.Visibility(this, &SClanListContainerImpl::GetClansListOpenVisbility, true)
+						]
+						+ SHorizontalBox::Slot()
+						.HAlign(HAlign_Left)
+						.VAlign(VAlign_Center)
+						.AutoWidth()
+						[
+							SNew(SButton)
+							.OnClicked(this, &SClanListContainerImpl::HandleShowClansClicked)
+							.ButtonStyle(&FriendStyle.FriendListCloseButtonStyle)
+							.Visibility(this, &SClanListContainerImpl::GetClansListOpenVisbility, false)
+						]
+						+ SHorizontalBox::Slot()
+						.HAlign(HAlign_Left)
+						.VAlign(VAlign_Center)
+						.AutoWidth()
+						.Padding(FMargin(5, 0))
+						[
+							SNew(STextBlock)
+							.ColorAndOpacity(FLinearColor::White)
+							.Font(FriendStyle.FriendsFontStyleUserLarge)
+							.Text(GetListName())
+						]
+					]
+				]
 			]
 			+ SVerticalBox::Slot()
 			.AutoHeight()
 			[
 				SNew(SBox)
 				.Padding(FMargin(0.0f, 2.0f, 0.0f, 8.0f))
+				.Visibility(this, &SClanListContainerImpl::GetClansListOpenVisbility, false)
 				[
-					CreateList()
+					SNew(SClanList, FClanListViewModelFactory::Create(ViewModel.ToSharedRef(), ClanListType)).FriendStyle(&FriendStyle)
 				]
 			]
 		]);
@@ -49,22 +90,38 @@ public:
 
 private:
 
-	TSharedRef<SWidget> CreateList()
+	FText GetListName() const
 	{
-		return SNew(SClanList, ViewModel.ToSharedRef())
-			.FriendStyle(&FriendStyle)
-			.Method(MenuMethod);
+		return EClanDisplayLists::ToFText(ClanListType);
+	}
+
+	EVisibility GetClansListOpenVisbility(bool bOpenButton) const
+	{
+		if (bOpenButton)
+		{
+			return ClansVisibility == EVisibility::Visible ? EVisibility::Collapsed : EVisibility::Visible;
+		}
+		return ClansVisibility;
+	}
+
+	FReply HandleShowClansClicked()
+	{
+		ClansVisibility = ClansVisibility == EVisibility::Visible ? EVisibility::Collapsed : EVisibility::Visible;
+		return FReply::Handled();
 	}
 
-private:
 	/** Holds the style to use when making the widget. */
 	FFriendsAndChatStyle FriendStyle;
 
-	EVisibility  FriendsVisibility;
+	EVisibility  ClansVisibility;
 
-	TSharedPtr<FClanListViewModel> ViewModel;
+	TSharedPtr<FClanViewModel> ViewModel;
 
 	EPopupMethod MenuMethod;
+
+	EClanDisplayLists::Type ClanListType;
+
+	bool bListOpen;
 };
 
 TSharedRef<SClanListContainer> SClanListContainer::New()
diff --git a/Engine/Source/Developer/FriendsAndChat/Private/UI/Widgets/SClanListContainer.h b/Engine/Source/Developer/FriendsAndChat/Private/UI/Widgets/SClanListContainer.h
--- a/Engine/Source/Developer/FriendsAndChat/Private/UI/Widgets/SClanListContainer.h
+++ b/Engine/Source/Developer/FriendsAndChat/Private/UI/Widgets/SClanListContainer.h
@@ -2,6 +2,11 @@
 
 #pragma once
 
+namespace EClanDisplayLists
+{
+	enum Type : uint8;
+}
+
 class SClanListContainer : public SUserWidget
 {
 public:
@@ -9,7 +14,8 @@ public:
 	SLATE_USER_ARGS(SClanListContainer) { }
 	SLATE_ARGUMENT(const FFriendsAndChatStyle*, FriendStyle)
 	SLATE_ARGUMENT(EPopupMethod, Method)
+	SLATE_ARGUMENT(EClanDisplayLists::Type, ClanListType)
 	SLATE_END_ARGS()
 
-	virtual void Construct(const FArguments& InArgs, const TSharedRef<class FClanListViewModel>& ViewModel) = 0;
+	virtual void Construct(const FArguments& InArgs, const TSharedRef<class FClanViewModel>& ViewModel) = 0;
 };
diff --git a/Engine/Source/Developer/FriendsAndChat/Public/FriendsAndChatStyle.h b/Engine/Source/Developer/FriendsAndChat/Public/FriendsAndChatStyle.h
--- a/Engine/Source/Developer/FriendsAndChat/Public/FriendsAndChatStyle.h
+++ b/Engine/Source/Developer/FriendsAndChat/Public/FriendsAndChatStyle.h
@@ -446,6 +446,14 @@ struct FRIENDSANDCHAT_API FFriendsAndChatStyle
 	UPROPERTY(EditAnywhere, Category = Appearance)
 	TEnumAsByte<EChatChannelStyle::Type> ChatChannelStyle;
 	FFriendsAndChatStyle&  SetChatWindowStyle(const EChatChannelStyle::Type Style);
+
+	UPROPERTY(EditAnywhere, Category = Appearance)
+	FSlateBrush ClanDetailsBrush;
+	FFriendsAndChatStyle& SetClanDetailsBrush(const FSlateBrush& Brush);
+
+	UPROPERTY(EditAnywhere, Category = Appearance)
+	FSlateBrush ClanMembersBrush;
+	FFriendsAndChatStyle& SetClanMembersBrush(const FSlateBrush& Brush);
 };
 
 /** Manages the style which provides resources for the rich text widget. */
diff --git a/Engine/Source/Runtime/Online/OnlineSubsystem/Public/Interfaces/OnlineChatInterface.h b/Engine/Source/Runtime/Online/OnlineSubsystem/Public/Interfaces/OnlineChatInterface.h
--- a/Engine/Source/Runtime/Online/OnlineSubsystem/Public/Interfaces/OnlineChatInterface.h
+++ b/Engine/Source/Runtime/Online/OnlineSubsystem/Public/Interfaces/OnlineChatInterface.h
@@ -22,6 +22,35 @@ public:
 	virtual const FString& GetSubject() const = 0;
 	virtual bool IsPrivate() const = 0;
 	virtual bool IsJoined() const = 0;
+	virtual const class FChatRoomConfig& GetRoomConfig() const = 0;
+};
+
+/**
+ * Configuration for creating a chat room
+ */
+class FChatRoomConfig
+{
+public:
+	FChatRoomConfig()
+		: bMembersOnly(false)
+		, bHidden(false)
+		, bPasswordRequired(false)
+		, bPersistent(false)
+		, bAllowMemberInvites(false)
+		, bLoggingEnabled(false)
+		, MessageHistory(0)
+		, MaxMembers(0)
+	{}
+
+	bool bMembersOnly;
+	bool bHidden;
+	bool bPasswordRequired;
+	bool bPersistent;
+	bool bAllowMemberInvites;
+	bool bLoggingEnabled;
+	int32 MessageHistory;
+	int32 MaxMembers;
+	FString PubSubNode;
 };
 
 /**
@@ -127,6 +156,17 @@ typedef FOnChatPrivateMessageReceived::FDelegate FOnChatPrivateMessageReceivedDe
 class IOnlineChat
 {
 public:
+
+	/**
+	 * Kick off request for creating a chat room with a provided configuration
+	 * 
+	 * @param UserId id of user that is joining
+	 * @param RoomConfig configuration for the room
+	 * @param Nickname display name for the chat room. Name must be unique and is reserved for duration of join
+	 *
+	 * @return if successfully started the async operation
+	 */
+	virtual bool CreateRoom(const FUniqueNetId& UserId, const FChatRoomConfig& RoomConfig, const FString& Nickname) = 0;
 	
 	/**
 	 * Kick off request for joining a public chat room
