#UE4 Create clan and friends factory on demand for now.

diff --git a/Engine/Source/Developer/FriendsAndChat/Private/FriendsAndChatManager.cpp b/Engine/Source/Developer/FriendsAndChat/Private/FriendsAndChatManager.cpp
--- a/Engine/Source/Developer/FriendsAndChat/Private/FriendsAndChatManager.cpp
+++ b/Engine/Source/Developer/FriendsAndChat/Private/FriendsAndChatManager.cpp
@@ -422,7 +422,18 @@ TSharedPtr< SWidget > FFriendsAndChatManager::GenerateFriendsListWidget( const F
 {
 	if ( !FriendListWidget.IsValid() )
 	{
+		if (!ClanRepository.IsValid())
+		{
+			ClanRepository = FClanRepositoryFactory::Create();
+		}
 		check(ClanRepository.IsValid())
+
+		if (!FriendsListFactory.IsValid())
+		{
+			FriendsListFactory = FFriendListFactoryFactory::Create();
+		}
+		check(FriendsListFactory.IsValid());
+
 		Style = *InStyle;
 		FFriendsAndChatModuleStyle::Initialize(Style);
 		SAssignNew(FriendListWidget, SOverlay)
