#UE4 Adjusting code to match what is in UE4-LauncherDev. This bug was already fixed bug perforce resolves unfixed it. Redoing the fix in main to get things in sync again. Fixes UE-13314

diff --git a/Engine/Source/Developer/FriendsAndChat/Private/FriendsAndChatManager.cpp b/Engine/Source/Developer/FriendsAndChat/Private/FriendsAndChatManager.cpp
--- a/Engine/Source/Developer/FriendsAndChat/Private/FriendsAndChatManager.cpp
+++ b/Engine/Source/Developer/FriendsAndChat/Private/FriendsAndChatManager.cpp
@@ -1108,14 +1108,11 @@ bool FFriendsAndChatManager::JoinGameAllowed(FString ClientID)
 	}
 	else
 	{
-		TSharedPtr<IFriendsApplicationViewModel>* FriendsApplicationViewModelRawPtr = ApplicationViewModels.Find(ClientID);
-		if (FriendsApplicationViewModelRawPtr != nullptr)
+		TSharedPtr<IFriendsApplicationViewModel>* FriendsApplicationViewModel = ApplicationViewModels.Find(ClientID);
+		if (FriendsApplicationViewModel != nullptr &&
+			(*FriendsApplicationViewModel).IsValid())
 		{
-			TSharedPtr<IFriendsApplicationViewModel> FriendsApplicationViewModel = *FriendsApplicationViewModelRawPtr;
-			if (FriendsApplicationViewModel.IsValid())
-			{
-				return FriendsApplicationViewModel->IsAppJoinable();
-			}
+			return (*FriendsApplicationViewModel)->IsAppJoinable();
 		}
 	}
 	return false;
@@ -2092,11 +2089,12 @@ void FFriendsAndChatManager::AcceptGameInvite(const TSharedPtr<IFriendItem>& Fri
 	// notify for further processing of join game request 
 	OnFriendsJoinGame().Broadcast(*FriendItem->GetUniqueID(), *FriendItem->GetGameSessionId());
 
-	TSharedPtr<IFriendsApplicationViewModel> FriendsApplicationViewModel = *ApplicationViewModels.Find(FriendItem->GetClientId());
-	if (FriendsApplicationViewModel.IsValid())
+	TSharedPtr<IFriendsApplicationViewModel>* FriendsApplicationViewModel = ApplicationViewModels.Find(FriendItem->GetClientId());
+	if (FriendsApplicationViewModel != nullptr &&
+		(*FriendsApplicationViewModel).IsValid())
 	{
 		const FString AdditionalCommandline = TEXT("-invitesession=") + FriendItem->GetGameSessionId()->ToString() + TEXT(" -invitefrom=") + FriendItem->GetUniqueID()->ToString();
-		FriendsApplicationViewModel->LaunchFriendApp(AdditionalCommandline);
+		(*FriendsApplicationViewModel)->LaunchFriendApp(AdditionalCommandline);
 	}
 
 	Analytics.RecordGameInvite(*FriendItem->GetUniqueID(), TEXT("Social.GameInvite.Accept"));
