[INTEGRATE] merge ue4-launcheDev->ue cl 2505085 & Cl 2505126

diff --git a/Engine/Source/Developer/FriendsAndChat/Private/FriendsMessageManager.cpp b/Engine/Source/Developer/FriendsAndChat/Private/FriendsMessageManager.cpp
--- a/Engine/Source/Developer/FriendsAndChat/Private/FriendsMessageManager.cpp
+++ b/Engine/Source/Developer/FriendsAndChat/Private/FriendsMessageManager.cpp
@@ -106,6 +106,8 @@ public:
 			{
 				TSharedPtr< FFriendChatMessage > ChatItem = MakeShareable(new FFriendChatMessage());
 				ChatItem->FromName = FText::FromString(OnlineSub->GetIdentityInterface()->GetPlayerNickname(*LoggedInUser.Get()));
+				TSharedPtr<IFriendItem> FoundFriend = FFriendsAndChatManager::Get()->FindUser(*UserID.Get());
+				ChatItem->ToName = FText::FromString(*FoundFriend->GetName());
 				ChatItem->Message = MessageText;
 				ChatItem->MessageType = EChatMessageType::Whisper;
 				ChatItem->MessageTime = FDateTime::UtcNow();
diff --git a/Engine/Source/Developer/FriendsAndChat/Private/FriendsMessageManager.h b/Engine/Source/Developer/FriendsAndChat/Private/FriendsMessageManager.h
--- a/Engine/Source/Developer/FriendsAndChat/Private/FriendsMessageManager.h
+++ b/Engine/Source/Developer/FriendsAndChat/Private/FriendsMessageManager.h
@@ -40,6 +40,7 @@ struct FFriendChatMessage
 {
 	EChatMessageType::Type MessageType;
 	FText FromName;
+	FText ToName;
 	FText Message;
 	FDateTime MessageTime;
 	FDateTime ExpireTime;
diff --git a/Engine/Source/Developer/FriendsAndChat/Private/Models/ChatItemViewModel.cpp b/Engine/Source/Developer/FriendsAndChat/Private/Models/ChatItemViewModel.cpp
--- a/Engine/Source/Developer/FriendsAndChat/Private/Models/ChatItemViewModel.cpp
+++ b/Engine/Source/Developer/FriendsAndChat/Private/Models/ChatItemViewModel.cpp
@@ -61,11 +61,21 @@ public:
 		return ChatMessages[0]->FromName;
 	}
 
+	virtual FText GetRecipientName() const override
+	{
+		return ChatMessages[0]->ToName;
+	}
+
 	virtual const TSharedPtr<FUniqueNetId> GetSenderID() const override
 	{
 		return ChatMessages[0]->SenderId;
 	}
 
+	virtual const TSharedPtr<FUniqueNetId> GetRecipientID() const override
+	{
+		return ChatMessages[0]->RecipientId;
+	}
+
 	const bool IsFromSelf() const override
 	{
 		return ChatMessages[0]->bIsFromSelf;
diff --git a/Engine/Source/Developer/FriendsAndChat/Private/Models/ChatItemViewModel.h b/Engine/Source/Developer/FriendsAndChat/Private/Models/ChatItemViewModel.h
--- a/Engine/Source/Developer/FriendsAndChat/Private/Models/ChatItemViewModel.h
+++ b/Engine/Source/Developer/FriendsAndChat/Private/Models/ChatItemViewModel.h
@@ -19,8 +19,11 @@ public:
 	virtual void AddMessage(const TSharedRef<FFriendChatMessage>& ChatMessage) = 0;
 
 	virtual FText GetSenderName() const = 0;
-
 	virtual const TSharedPtr<FUniqueNetId> GetSenderID() const = 0;
+
+	virtual FText GetRecipientName() const = 0;
+	virtual const TSharedPtr<FUniqueNetId> GetRecipientID() const = 0;
+
 	virtual const bool IsFromSelf() const = 0;
 
 	DECLARE_EVENT_OneParam(FChatItemViewModel, FChangedEvent, const TSharedRef<FChatItemViewModel>&)
diff --git a/Engine/Source/Developer/FriendsAndChat/Private/Models/ChatViewModel.cpp b/Engine/Source/Developer/FriendsAndChat/Private/Models/ChatViewModel.cpp
--- a/Engine/Source/Developer/FriendsAndChat/Private/Models/ChatViewModel.cpp
+++ b/Engine/Source/Developer/FriendsAndChat/Private/Models/ChatViewModel.cpp
@@ -255,9 +255,10 @@ public:
 			bool bFoundUser = false;
 			TSharedPtr< IFriendItem > ExistingFriend = NULL;
 
-			if(ChatItemSelected->GetSenderID().IsValid())
+			const TSharedPtr<FUniqueNetId> ChatID = ChatItemSelected->IsFromSelf() ? ChatItemSelected->GetRecipientID() : ChatItemSelected->GetSenderID();
+			if(ChatID.IsValid())
 			{
-				ExistingFriend = FFriendsAndChatManager::Get()->FindUser(*ChatItemSelected->GetSenderID().Get());
+				ExistingFriend = FFriendsAndChatManager::Get()->FindUser(*ChatID.Get());
 			}
 
 			if(ExistingFriend.IsValid() && ExistingFriend->GetInviteStatus() == EInviteStatus::Accepted && !bLockedChannel)
diff --git a/Engine/Source/Developer/FriendsAndChat/Private/UI/Widgets/SChatItem.cpp b/Engine/Source/Developer/FriendsAndChat/Private/UI/Widgets/SChatItem.cpp
--- a/Engine/Source/Developer/FriendsAndChat/Private/UI/Widgets/SChatItem.cpp
+++ b/Engine/Source/Developer/FriendsAndChat/Private/UI/Widgets/SChatItem.cpp
@@ -365,7 +365,15 @@ private:
 		if (bRequiresTextUpdate)
 		{
 			FFormatNamedArguments Args;
-			Args.Add(TEXT("SenderName"), ViewModel->GetSenderName());
+			if(ViewModel->IsFromSelf() && ViewModel->GetMessageType() != EChatMessageType::Global)
+			{
+				Args.Add(TEXT("SenderName"), ViewModel->GetRecipientName());
+			}
+			else
+			{
+				Args.Add(TEXT("SenderName"), ViewModel->GetSenderName());
+			}
+
 			Args.Add(TEXT("NameStyle"), FText::FromString(GetTextHyperlinkStyle()));
 			Args.Add(TEXT("Message"), TextToSet);
 			CurrentListMessageText = FText::Format(NSLOCTEXT("FChatItemViewModel", "DisplayNameAndMessage", "<a id=\"UserName\" style=\"{NameStyle}\">{SenderName}</>: {Message}"), Args);
diff --git a/Engine/Source/Developer/FriendsAndChat/Private/UI/Widgets/SChatWindow.cpp b/Engine/Source/Developer/FriendsAndChat/Private/UI/Widgets/SChatWindow.cpp
--- a/Engine/Source/Developer/FriendsAndChat/Private/UI/Widgets/SChatWindow.cpp
+++ b/Engine/Source/Developer/FriendsAndChat/Private/UI/Widgets/SChatWindow.cpp
@@ -91,94 +91,94 @@ public:
 						SNew(SBorder)
 						.BorderImage(&FriendStyle.ChatFooterBrush)
 						.Padding(FMargin(10, 10, 10, 10))
-				[
-					SNew(SHorizontalBox)
-					.Visibility(this, &SChatWindowImpl::GetActionVisibility)
-					+SHorizontalBox::Slot()
-					.AutoWidth()
-					.HAlign(HAlign_Left)
-					.VAlign(VAlign_Fill)
-					[
-						SAssignNew(ActionMenu, SMenuAnchor)
-						.Placement(MenuMethod != EPopupMethod::UseCurrentWindow ? EMenuPlacement::MenuPlacement_AboveAnchor : EMenuPlacement::MenuPlacement_BelowAnchor)  	// MenuPlacement_BelowAnchor is a workaround until a workarea bug is fixed in SlateApplication menu placement code
-						.IsEnabled(this, &SChatWindowImpl::HasMenuOptions)
-						.Method(EPopupMethod::UseCurrentWindow)
-						.OnGetMenuContent(this, &SChatWindowImpl::GetMenuContent)
-						.Visibility(this, &SChatWindowImpl::GetChatChannelVisibility)
-						[
-							SNew(SButton)
-							.ButtonStyle(&FriendStyle.FriendGeneralButtonStyle)
-							.ContentPadding(FMargin(2, 2.5))
-							.OnClicked(this, &SChatWindowImpl::HandleActionDropDownClicked)
-							.VAlign(VAlign_Center)
-							.HAlign(HAlign_Center)
-							.Cursor(EMouseCursor::Hand)
-							[
-								SNew(SImage)
-								.Image(this, &SChatWindowImpl::GetChatChannelIcon)
-							]
-						]
-					]
-					+SHorizontalBox::Slot()
-					[
-						SNew(SVerticalBox)
-						+ SVerticalBox::Slot()
 						[
 							SNew(SHorizontalBox)
-							.Visibility(this, &SChatWindowImpl::GetChatEntryVisibility)
+							.Visibility(this, &SChatWindowImpl::GetActionVisibility)
 							+SHorizontalBox::Slot()
+							.AutoWidth()
 							.HAlign(HAlign_Left)
 							.VAlign(VAlign_Fill)
-							.AutoWidth()
 							[
-								SNew(SBorder)
-								.BorderImage(&FriendStyle.ChatContainerBackground)
-								.BorderBackgroundColor(FLinearColor(FColor(255, 255, 255, 128)))
+								SAssignNew(ActionMenu, SMenuAnchor)
+								.Placement(MenuMethod != EPopupMethod::UseCurrentWindow ? EMenuPlacement::MenuPlacement_AboveAnchor : EMenuPlacement::MenuPlacement_BelowAnchor)  	// MenuPlacement_BelowAnchor is a workaround until a workarea bug is fixed in SlateApplication menu placement code
+								.IsEnabled(this, &SChatWindowImpl::HasMenuOptions)
+								.Method(EPopupMethod::UseCurrentWindow)
+								.OnGetMenuContent(this, &SChatWindowImpl::GetMenuContent)
+								.Visibility(this, &SChatWindowImpl::GetChatChannelVisibility)
 								[
-									SNew(SHorizontalBox)
-									.Visibility(this, &SChatWindowImpl::GetFriendNameVisibility)
-									+SHorizontalBox::Slot()
-									.Padding(5,0)
-									.AutoWidth()
+									SNew(SButton)
+									.ButtonStyle(&FriendStyle.FriendGeneralButtonStyle)
+									.ContentPadding(FMargin(2, 2.5))
+									.OnClicked(this, &SChatWindowImpl::HandleActionDropDownClicked)
 									.VAlign(VAlign_Center)
+									.HAlign(HAlign_Center)
+									.Cursor(EMouseCursor::Hand)
 									[
-										SNew(STextBlock)
-										.Font(FriendStyle.FriendsFontStyleSmallBold)
-										.Text(ViewModelPtr, &FChatViewModel::GetChatGroupText)
+										SNew(SImage)
+										.Image(this, &SChatWindowImpl::GetChatChannelIcon)
 									]
+								]
+							]
+							+SHorizontalBox::Slot()
+							[
+								SNew(SVerticalBox)
+								+ SVerticalBox::Slot()
+								[
+									SNew(SHorizontalBox)
+									.Visibility(this, &SChatWindowImpl::GetChatEntryVisibility)
 									+SHorizontalBox::Slot()
-									.Padding(0, 0, 5, 0)
+									.HAlign(HAlign_Left)
+									.VAlign(VAlign_Fill)
 									.AutoWidth()
-									.VAlign(VAlign_Center)
 									[
-										SNew(SButton)
-										.Visibility(this, &SChatWindowImpl::GetFriendActionVisibility)
-										.ButtonStyle(&FriendStyle.FriendGeneralButtonStyle)
-										.ContentPadding(FMargin(-2,3))
-										.OnClicked(this, &SChatWindowImpl::HandleFriendActionDropDownClicked)
+										SNew(SBorder)
+										.BorderImage(&FriendStyle.ChatContainerBackground)
+										.BorderBackgroundColor(FLinearColor(FColor(255, 255, 255, 128)))
 										[
-											SAssignNew(ChatItemActionMenu, SMenuAnchor)
-											.Placement(EMenuPlacement::MenuPlacement_BelowAnchor)	// MenuPlacement_BelowAnchor is a workaround until a workarea bug is fixed in SlateApplication menu placement code
-											.Method(EPopupMethod::UseCurrentWindow)
-											.OnGetMenuContent(this, &SChatWindowImpl::GetFriendActionMenu)
+											SNew(SHorizontalBox)
+											.Visibility(this, &SChatWindowImpl::GetFriendNameVisibility)
+											+SHorizontalBox::Slot()
+											.Padding(5,0)
+											.AutoWidth()
+											.VAlign(VAlign_Center)
 											[
-												SNew(SImage)
-												.Image(&FriendStyle.FriendsCalloutBrush)
+												SNew(STextBlock)
+												.Font(FriendStyle.FriendsFontStyleSmallBold)
+												.Text(ViewModelPtr, &FChatViewModel::GetChatGroupText)
+											]
+											+SHorizontalBox::Slot()
+											.Padding(0, 0, 5, 0)
+											.AutoWidth()
+											.VAlign(VAlign_Center)
+											[
+												SNew(SButton)
+												.Visibility(this, &SChatWindowImpl::GetFriendActionVisibility)
+												.ButtonStyle(&FriendStyle.FriendGeneralButtonStyle)
+												.ContentPadding(FMargin(-2,3))
+												.OnClicked(this, &SChatWindowImpl::HandleFriendActionDropDownClicked)
+												[
+													SAssignNew(ChatItemActionMenu, SMenuAnchor)
+													.Placement(EMenuPlacement::MenuPlacement_BelowAnchor)	// MenuPlacement_BelowAnchor is a workaround until a workarea bug is fixed in SlateApplication menu placement code
+													.Method(EPopupMethod::UseCurrentWindow)
+													.OnGetMenuContent(this, &SChatWindowImpl::GetFriendActionMenu)
+													[
+														SNew(SImage)
+														.Image(&FriendStyle.FriendsCalloutBrush)
+													]
+												]
 											]
 										]
 									]
-								]
-							]
-							+SHorizontalBox::Slot()
-							[
+									+SHorizontalBox::Slot()
+									[
 										SAssignNew(ChatTextBox, SMultiLineEditableTextBox)
-								.Style(&FriendStyle.ChatEditableTextStyle)
+										.Style(&FriendStyle.ChatEditableTextStyle)
 										.TextStyle(&FriendStyle.TextStyle)
-								.ClearKeyboardFocusOnCommit(false)
-								.OnTextCommitted(this, &SChatWindowImpl::HandleChatEntered)
-								.HintText(InArgs._ActivationHintText)
-								.OnTextChanged(this, &SChatWindowImpl::OnChatTextChanged)
-								.IsEnabled(this, &SChatWindowImpl::IsChatEntryEnabled)
+										.ClearKeyboardFocusOnCommit(false)
+										.OnTextCommitted(this, &SChatWindowImpl::HandleChatEntered)
+										.HintText(InArgs._ActivationHintText)
+										.OnTextChanged(this, &SChatWindowImpl::OnChatTextChanged)
+										.IsEnabled(this, &SChatWindowImpl::IsChatEntryEnabled)
 										.ModiferKeyForNewLine(EModifierKey::Shift)
 										.WrapTextAt(this, &SChatWindowImpl::GetMessageBoxWrapWidth)
 									]
@@ -198,17 +198,17 @@ public:
 										.AutoWrapText(true)
 										.Font(FriendStyle.FriendsFontStyleSmallBold)
 										.Text(ViewModelPtr, &FChatViewModel::GetChatDisconnectText)
-							]
-						]
-						+SVerticalBox::Slot()
-						[
-							SNew(SBorder)
-							.Visibility(this, &SChatWindowImpl::GetConfirmationVisibility)
-							.BorderImage(&FriendStyle.ChatContainerBackground)
-							.BorderBackgroundColor(FLinearColor(FColor(255, 255, 255, 128)))
-							[
-								SNew(SHorizontalBox)
-								+SHorizontalBox::Slot()
+									]
+								]
+								+SVerticalBox::Slot()
+								[
+									SNew(SBorder)
+									.Visibility(this, &SChatWindowImpl::GetConfirmationVisibility)
+									.BorderImage(&FriendStyle.ChatContainerBackground)
+									.BorderBackgroundColor(FLinearColor(FColor(255, 255, 255, 128)))
+									[
+										SNew(SHorizontalBox)
+										+SHorizontalBox::Slot()
 										[
 											SNew(SVerticalBox)
 											+ SVerticalBox::Slot()
@@ -216,12 +216,12 @@ public:
 											.VAlign(VAlign_Top)
 											.HAlign(HAlign_Left)
 											.AutoHeight()
-								[
-									SNew(STextBlock)
+											[
+												SNew(STextBlock)
 												.Font(FriendStyle.FriendsFontStyleBold)
 												.ColorAndOpacity(FriendStyle.DefaultFontColor)
-									.Text(ViewModelPtr, &FChatViewModel::GetChatGroupText)
-								]
+												.Text(ViewModelPtr, &FChatViewModel::GetChatGroupText)
+											]
 											+ SVerticalBox::Slot()
 											.VAlign(VAlign_Top)
 											.HAlign(HAlign_Left)
@@ -233,51 +233,51 @@ public:
 												.MinDesiredSlotWidth(150.0f)
 												.MinDesiredSlotHeight(30.0f)
 												+ SUniformGridPanel::Slot(0,0)
-								[
+												[
 													SAssignNew(SendFriendRequestButton, SButton)
 													.ButtonStyle(SFriendsList::GetActionButtonStyle(FriendStyle, EFriendActionType::ToActionLevel(EFriendActionType::SendFriendRequest)))
 													.OnClicked(this, &SChatWindowImpl::HandleFriendActionClicked, EFriendActionType::SendFriendRequest)
-									.Visibility(ViewModelPtr, &FChatViewModel::GetFriendRequestVisibility)
+													.Visibility(ViewModelPtr, &FChatViewModel::GetFriendRequestVisibility)
 													.HAlign(HAlign_Center)
-									.VAlign(VAlign_Center)
+													.VAlign(VAlign_Center)
 													.Cursor(EMouseCursor::Hand)
-									[
-										SNew(STextBlock)
-										.Font(FriendStyle.FriendsFontStyleSmallBold)
+													[
+														SNew(STextBlock)
+														.Font(FriendStyle.FriendsFontStyleSmallBold)
 														.Text(EFriendActionType::ToText(EFriendActionType::SendFriendRequest))
 														.ColorAndOpacity(FSlateColor::UseForeground())
-									]
-								]
+													]
+												]
 												+ SUniformGridPanel::Slot(0, 0)
-								[
+												[
 													SAssignNew(InviteToGameButton, SButton)
-									.Visibility(ViewModelPtr, &FChatViewModel::GetInviteToGameVisibility)
+													.Visibility(ViewModelPtr, &FChatViewModel::GetInviteToGameVisibility)
 													.ButtonStyle(SFriendsList::GetActionButtonStyle(FriendStyle, EFriendActionType::ToActionLevel(EFriendActionType::InviteToGame)))
-									.OnClicked(this, &SChatWindowImpl::HandleFriendActionClicked, EFriendActionType::InviteToGame)
-									.ContentPadding(5)
-									.VAlign(VAlign_Center)
-									.HAlign(HAlign_Center)
+													.OnClicked(this, &SChatWindowImpl::HandleFriendActionClicked, EFriendActionType::InviteToGame)
+													.ContentPadding(5)
+													.VAlign(VAlign_Center)
+													.HAlign(HAlign_Center)
 													.Cursor(EMouseCursor::Hand)
-									[
-										SNew(STextBlock)
-										.Font(FriendStyle.FriendsFontStyleSmallBold)
+													[
+														SNew(STextBlock)
+														.Font(FriendStyle.FriendsFontStyleSmallBold)
 														.Text(EFriendActionType::ToText(EFriendActionType::InviteToGame))
 														.ColorAndOpacity(FSlateColor::UseForeground())
-									]
-								]
+													]
+												]
 												+ SUniformGridPanel::Slot(0, 0)
-								[
+												[
 													SAssignNew(OpenWhisperButton, SButton)
-									.Visibility(ViewModelPtr, &FChatViewModel::GetOpenWhisperVisibility)
+													.Visibility(ViewModelPtr, &FChatViewModel::GetOpenWhisperVisibility)
 													.ButtonStyle(SFriendsList::GetActionButtonStyle(FriendStyle, EFriendActionType::ToActionLevel(EFriendActionType::Whisper)))
-									.OnClicked(this, &SChatWindowImpl::HandleFriendActionClicked, EFriendActionType::Whisper)
-									.ContentPadding(5)
-									.VAlign(VAlign_Center)
-									.HAlign(HAlign_Center)
+													.OnClicked(this, &SChatWindowImpl::HandleFriendActionClicked, EFriendActionType::Whisper)
+													.ContentPadding(5)
+													.VAlign(VAlign_Center)
+													.HAlign(HAlign_Center)
 													.Cursor(EMouseCursor::Hand)
-									[
-										SNew(STextBlock)
-										.Font(FriendStyle.FriendsFontStyleSmallBold)
+													[
+														SNew(STextBlock)
+														.Font(FriendStyle.FriendsFontStyleSmallBold)
 														.Text(EFriendActionType::ToText(EFriendActionType::Whisper))
 														.ColorAndOpacity(FSlateColor::UseForeground())
 													]
@@ -330,16 +330,16 @@ public:
 														.ColorAndOpacity(FSlateColor::UseForeground())
 													]
 												]
-									]
-								]
-								+ SHorizontalBox::Slot()
+											]
+										]
+										+ SHorizontalBox::Slot()
 										.HAlign(HAlign_Left)
-								.VAlign(VAlign_Center)
-								.AutoWidth()
-								[
-									SNew(SButton)
-									.ButtonStyle(&FriendStyle.AddFriendCloseButtonStyle)
-									.OnClicked(this, &SChatWindowImpl::CancelActionClicked)
+										.VAlign(VAlign_Center)
+										.AutoWidth()
+										[
+											SNew(SButton)
+											.ButtonStyle(&FriendStyle.AddFriendCloseButtonStyle)
+											.OnClicked(this, &SChatWindowImpl::CancelActionClicked)
 											.Cursor(EMouseCursor::Hand)
 										]
 									]
@@ -447,11 +447,11 @@ public:
 
 				if (FadeAnimation.IsPlaying())
 				{
-				    TArray<TSharedRef<FChatItemViewModel>>& FilteredChatList = SharedChatViewModel->GetMessages();
-				    if (FilteredChatList.Num() > 0)
-				    {
-					    ChatList->RequestScrollIntoView(FilteredChatList.Last());
-				    }
+					TArray<TSharedRef<FChatItemViewModel>>& FilteredChatList = SharedChatViewModel->GetMessages();
+					if (FilteredChatList.Num() > 0)
+					{
+						ChatList->RequestScrollIntoView(FilteredChatList.Last());
+					}
 				}
 			}
 		}
@@ -505,22 +505,22 @@ private:
 
 		if (!SharedChatViewModel->IsChatChannelLocked())
 		{
-		return SNew(STableRow< TSharedPtr<SWidget> >, OwnerTable)
-		[
-			SNew(SButton)
-			.ButtonStyle(FCoreStyle::Get(), "NoBorder")
-			.OnClicked(SharedChatViewModel.Get(), &FChatViewModel::HandleSelectionChanged, ChatMessage)
-			.VAlign(VAlign_Center)
+			return SNew(STableRow< TSharedPtr<SWidget> >, OwnerTable)
 			[
-					Content
-			]
-		];
-	}
+				SNew(SButton)
+				.ButtonStyle(FCoreStyle::Get(), "NoBorder")
+				.OnClicked(SharedChatViewModel.Get(), &FChatViewModel::HandleSelectionChanged, ChatMessage)
+				.VAlign(VAlign_Center)
+				[
+						Content
+				]
+			];
+		}
 
 		return SNew(STableRow< TSharedPtr<SWidget> >, OwnerTable)
-			[
-				Content
-			];
+		[
+			Content
+		];
 	}
 
 	FReply HandleActionDropDownClicked() const
@@ -738,7 +738,7 @@ private:
 		if (ActionType != EFriendActionType::Whisper)
 		{
 			SetFocus();
-		}		
+		}
 		return FReply::Handled();
 	}
 
