Fix up network chat to not allow user name to be selected

diff --git a/Engine/Source/Developer/FriendsAndChat/Private/UI/Widgets/SChatItem.cpp b/Engine/Source/Developer/FriendsAndChat/Private/UI/Widgets/SChatItem.cpp
--- a/Engine/Source/Developer/FriendsAndChat/Private/UI/Widgets/SChatItem.cpp
+++ b/Engine/Source/Developer/FriendsAndChat/Private/UI/Widgets/SChatItem.cpp
@@ -364,19 +364,26 @@ private:
 
 		if (bRequiresTextUpdate)
 		{
-			FFormatNamedArguments Args;
-			if(ViewModel->IsFromSelf() && ViewModel->GetMessageType() != EChatMessageType::Global)
+			if(ViewModel->GetMessageType() == EChatMessageType::Party)
 			{
-				Args.Add(TEXT("SenderName"), ViewModel->GetRecipientName());
+				CurrentListMessageText = TextToSet;
 			}
 			else
 			{
-				Args.Add(TEXT("SenderName"), ViewModel->GetSenderName());
+				FFormatNamedArguments Args;
+				if(ViewModel->IsFromSelf() && ViewModel->GetMessageType() != EChatMessageType::Global)
+				{
+					Args.Add(TEXT("SenderName"), ViewModel->GetRecipientName());
+				}
+				else
+				{
+					Args.Add(TEXT("SenderName"), ViewModel->GetSenderName());
+				}
+
+				Args.Add(TEXT("NameStyle"), FText::FromString(GetTextHyperlinkStyle()));
+				Args.Add(TEXT("Message"), TextToSet);
+				CurrentListMessageText = FText::Format(NSLOCTEXT("FChatItemViewModel", "DisplayNameAndMessage", "<a id=\"UserName\" style=\"{NameStyle}\">{SenderName}</>: {Message}"), Args);
 			}
-
-			Args.Add(TEXT("NameStyle"), FText::FromString(GetTextHyperlinkStyle()));
-			Args.Add(TEXT("Message"), TextToSet);
-			CurrentListMessageText = FText::Format(NSLOCTEXT("FChatItemViewModel", "DisplayNameAndMessage", "<a id=\"UserName\" style=\"{NameStyle}\">{SenderName}</>: {Message}"), Args);
 		}
 
 		return CurrentListMessageText;
