Social.Chat.Counts -> Social.Chat.Counts.2
* pivoted the channel info into separate events, added the channel type info.
New format:
Social.Chat.Counts.2  ; note the .2 extension since the original is already in the wild.
    Name:<String>   ; will be a user name or a channel name, depending on type.
    Type:<Private|Channel> ; private chat is with one user, channel chats have more than one,
                                             ; and is named after the channel somehow (not totally sure how).
    Count:<int> ; number of chats in that channel since the last flush.

diff --git a/Engine/Source/Developer/FriendsAndChat/Private/FriendsAndChatManager.cpp b/Engine/Source/Developer/FriendsAndChat/Private/FriendsAndChatManager.cpp
--- a/Engine/Source/Developer/FriendsAndChat/Private/FriendsAndChatManager.cpp
+++ b/Engine/Source/Developer/FriendsAndChat/Private/FriendsAndChatManager.cpp
@@ -2276,15 +2276,13 @@ void FFriendsAndChatAnalytics::RecordToggleChat(const FString& Channel, bool bEn
 
 void FFriendsAndChatAnalytics::RecordPrivateChat(const FString& ToUser)
 {
-	int32& Count = ChatCounts.FindOrAdd(ToUser);
+	int32& Count = PrivateChatCounts.FindOrAdd(ToUser);
 	Count += 1;
-	int32& TotalCount = ChatCounts.FindOrAdd(TEXT("TotalPrivate"));
-	TotalCount += 1;
 }
 
 void FFriendsAndChatAnalytics::RecordChannelChat(const FString& ToChannel)
 {
-	int32& Count = ChatCounts.FindOrAdd(ToChannel);
+	int32& Count = ChannelChatCounts.FindOrAdd(ToChannel);
 	Count += 1;
 }
 
@@ -2298,16 +2296,29 @@ void FFriendsAndChatAnalytics::FlushChatStats()
 			TSharedPtr<FUniqueNetId> UserId = OnlineIdentity->GetUniquePlayerId(0);
 			if (UserId.IsValid())
 			{
-				TArray<FAnalyticsEventAttribute> Attributes;
-				for (auto It = ChatCounts.CreateConstIterator(); It; ++It)
+				auto RecordSocialChatCountsEvents = [=](const TMap<FString, int32>& ChatCounts, const FString& ChatType)
 				{
-					Attributes.Add(FAnalyticsEventAttribute(It.Key(), It.Value()));
-				}
-				Provider->RecordEvent(TEXT("Social.Chat.Counts"), Attributes);
+					if (ChatCounts.Num())
+					{
+						TArray<FAnalyticsEventAttribute> Attributes;
+						for (const auto& Pair : ChannelChatCounts)
+						{
+							Attributes.Empty(3);
+							Attributes.Emplace(TEXT("Name"), Pair.Key);
+							Attributes.Emplace(TEXT("Type"), ChatType);
+							Attributes.Emplace(TEXT("Count"), Pair.Value);
+							Provider->RecordEvent("Social.Chat.Counts.2", Attributes);
+						}
+					}
+				};
+
+				RecordSocialChatCountsEvents(ChannelChatCounts, TEXT("Channel"));
+				RecordSocialChatCountsEvents(PrivateChatCounts, TEXT("Private"));
 			}
 		}
 	}
-	ChatCounts.Empty();
+	ChannelChatCounts.Empty();
+	PrivateChatCounts.Empty();
 }
 
 void FFriendsAndChatAnalytics::AddPresenceAttributes(const FUniqueNetId& UserId, TArray<FAnalyticsEventAttribute>& Attributes) const
diff --git a/Engine/Source/Developer/FriendsAndChat/Private/FriendsAndChatManager.h b/Engine/Source/Developer/FriendsAndChat/Private/FriendsAndChatManager.h
--- a/Engine/Source/Developer/FriendsAndChat/Private/FriendsAndChatManager.h
+++ b/Engine/Source/Developer/FriendsAndChat/Private/FriendsAndChatManager.h
@@ -101,7 +101,8 @@ private:
 	// cached analytics provider for pushing events
 	TSharedPtr<IAnalyticsProvider> Provider;
 	// map of chat id to # of messages sent
-	TMap<FString, int32> ChatCounts;
+	TMap<FString, int32> PrivateChatCounts;
+	TMap<FString, int32> ChannelChatCounts;
 };
 
 /**
