Fix UnrealLauncher shadow variables - batch 1

diff --git a/Engine/Source/Developer/FriendsAndChat/Private/ClanRepository.cpp b/Engine/Source/Developer/FriendsAndChat/Private/ClanRepository.cpp
--- a/Engine/Source/Developer/FriendsAndChat/Private/ClanRepository.cpp
+++ b/Engine/Source/Developer/FriendsAndChat/Private/ClanRepository.cpp
@@ -23,8 +23,8 @@ public:
 	/**
 	 * Constructor takes the required details.
 	 */
-	FClanMember(FClanMemberInfo MemberInfo)
-		:MemberInfo(MemberInfo)
+	FClanMember(FClanMemberInfo InMemberInfo)
+		:MemberInfo(InMemberInfo)
 	{ }
 
 public:
diff --git a/Engine/Source/Developer/FriendsAndChat/Private/Models/ChatDisplayOptionsViewModel.cpp b/Engine/Source/Developer/FriendsAndChat/Private/Models/ChatDisplayOptionsViewModel.cpp
--- a/Engine/Source/Developer/FriendsAndChat/Private/Models/ChatDisplayOptionsViewModel.cpp
+++ b/Engine/Source/Developer/FriendsAndChat/Private/Models/ChatDisplayOptionsViewModel.cpp
@@ -201,8 +201,8 @@ private:
 		OnChatListUpdated().Broadcast();
 	}
 
-	FChatDisplayOptionsViewModelImpl(const TSharedRef<FChatViewModel>& ChatViewModel)
-		: ChatViewModel(ChatViewModel)
+	FChatDisplayOptionsViewModelImpl(const TSharedRef<FChatViewModel>& InChatViewModel)
+		: ChatViewModel(InChatViewModel)
 		, bUseOverrideColor(false)
 		, bInGame(false)
 		, bAllowGlobalChat(true)
diff --git a/Engine/Source/Developer/FriendsAndChat/Private/Models/ChatViewModel.cpp b/Engine/Source/Developer/FriendsAndChat/Private/Models/ChatViewModel.cpp
--- a/Engine/Source/Developer/FriendsAndChat/Private/Models/ChatViewModel.cpp
+++ b/Engine/Source/Developer/FriendsAndChat/Private/Models/ChatViewModel.cpp
@@ -637,10 +637,10 @@ private:
 		}
 	}
 
-	FChatViewModelImpl(const TSharedRef<FFriendsMessageManager>& MessageManager)
+	FChatViewModelImpl(const TSharedRef<FFriendsMessageManager>& InMessageManager)
 		: SelectedViewChannel(EChatMessageType::Global)
 		, SelectedChatChannel(EChatMessageType::Global)
-		, MessageManager(MessageManager)
+		, MessageManager(InMessageManager)
 		, bInGame(false)
 		, bIsDisplayingGlobalChat(false)
 		, bEnableGlobalChat(true)
diff --git a/Engine/Source/Developer/FriendsAndChat/Private/Models/ClanInfoViewModel.cpp b/Engine/Source/Developer/FriendsAndChat/Private/Models/ClanInfoViewModel.cpp
--- a/Engine/Source/Developer/FriendsAndChat/Private/Models/ClanInfoViewModel.cpp
+++ b/Engine/Source/Developer/FriendsAndChat/Private/Models/ClanInfoViewModel.cpp
@@ -43,8 +43,8 @@ private:
 	}
 
 private:
-	FClanInfoViewModelImpl(const TSharedRef<IClanInfo>& ClanInfo)
-		: ClanInfo(ClanInfo) 
+	FClanInfoViewModelImpl(const TSharedRef<IClanInfo>& InClanInfo)
+		: ClanInfo(InClanInfo)
 	{
 	}
 
diff --git a/Engine/Source/Developer/FriendsAndChat/Private/Models/ClanListViewModel.cpp b/Engine/Source/Developer/FriendsAndChat/Private/Models/ClanListViewModel.cpp
--- a/Engine/Source/Developer/FriendsAndChat/Private/Models/ClanListViewModel.cpp
+++ b/Engine/Source/Developer/FriendsAndChat/Private/Models/ClanListViewModel.cpp
@@ -10,9 +10,9 @@ class FClanListViewModelImpl : public FClanListViewModel
 {
 public:
 
-	FClanListViewModelImpl(const TSharedRef<FClanViewModel>& ClanViewModel, EClanDisplayLists::Type ListType)
-		: ViewModel(ClanViewModel)
-		, ListType(ListType)
+	FClanListViewModelImpl(const TSharedRef<FClanViewModel>& InClanViewModel, EClanDisplayLists::Type InListType)
+		: ViewModel(InClanViewModel)
+		, ListType(InListType)
 	{
 	}
 
diff --git a/Engine/Source/Developer/FriendsAndChat/Private/Models/ClanMemberList.cpp b/Engine/Source/Developer/FriendsAndChat/Private/Models/ClanMemberList.cpp
--- a/Engine/Source/Developer/FriendsAndChat/Private/Models/ClanMemberList.cpp
+++ b/Engine/Source/Developer/FriendsAndChat/Private/Models/ClanMemberList.cpp
@@ -67,8 +67,8 @@ private:
 		OnFriendsListUpdated().Broadcast();
 	}
 
-	FClanMemberListImpl(TSharedRef<IClanInfo> ClanInfo)
-		: ClanInfo(ClanInfo)
+	FClanMemberListImpl(TSharedRef<IClanInfo> InClanInfo)
+		: ClanInfo(InClanInfo)
 	{}
 
 private:
diff --git a/Engine/Source/Developer/FriendsAndChat/Private/Models/ClanViewModel.cpp b/Engine/Source/Developer/FriendsAndChat/Private/Models/ClanViewModel.cpp
--- a/Engine/Source/Developer/FriendsAndChat/Private/Models/ClanViewModel.cpp
+++ b/Engine/Source/Developer/FriendsAndChat/Private/Models/ClanViewModel.cpp
@@ -45,8 +45,8 @@ private:
 	}
 
 private:
-	FClanViewModelImpl(const TSharedRef<IClanRepository>& ClanRepository)
-		: ClanRepository(ClanRepository) 
+	FClanViewModelImpl(const TSharedRef<IClanRepository>& InClanRepository)
+		: ClanRepository(InClanRepository)
 	{
 	}
 
diff --git a/Engine/Source/Developer/FriendsAndChat/Private/Models/DefaultFriendsList.cpp b/Engine/Source/Developer/FriendsAndChat/Private/Models/DefaultFriendsList.cpp
--- a/Engine/Source/Developer/FriendsAndChat/Private/Models/DefaultFriendsList.cpp
+++ b/Engine/Source/Developer/FriendsAndChat/Private/Models/DefaultFriendsList.cpp
@@ -92,8 +92,8 @@ private:
 		OnFriendsListUpdated().Broadcast();
 	}
 
-	FDefaultFriendListImpl(EFriendsDisplayLists::Type ListType)
-		: ListType(ListType)
+	FDefaultFriendListImpl(EFriendsDisplayLists::Type InListType)
+		: ListType(InListType)
 	{}
 
 private:
diff --git a/Engine/Source/Developer/FriendsAndChat/Private/Models/FriendListViewModel.cpp b/Engine/Source/Developer/FriendsAndChat/Private/Models/FriendListViewModel.cpp
--- a/Engine/Source/Developer/FriendsAndChat/Private/Models/FriendListViewModel.cpp
+++ b/Engine/Source/Developer/FriendsAndChat/Private/Models/FriendListViewModel.cpp
@@ -91,10 +91,10 @@ private:
 	}
 
 	FFriendListViewModelImpl( 
-		TSharedRef<IFriendList> FriendsListContainer,
-		EFriendsDisplayLists::Type ListType)
-		: FriendsListContainer(FriendsListContainer)
-		, ListType(ListType)
+		TSharedRef<IFriendList> InFriendsListContainer,
+		EFriendsDisplayLists::Type InListType)
+		: FriendsListContainer(InFriendsListContainer)
+		, ListType(InListType)
 	{
 	}
 
diff --git a/Engine/Source/Developer/FriendsAndChat/Private/Models/FriendRecentPlayerItems.h b/Engine/Source/Developer/FriendsAndChat/Private/Models/FriendRecentPlayerItems.h
--- a/Engine/Source/Developer/FriendsAndChat/Private/Models/FriendRecentPlayerItems.h
+++ b/Engine/Source/Developer/FriendsAndChat/Private/Models/FriendRecentPlayerItems.h
@@ -15,8 +15,8 @@ public:
 	 *
 	 * @param RecentPlayer The recent friend.
 	 */
-	FFriendRecentPlayerItem(TSharedPtr< FOnlineRecentPlayer > RecentPlayer)
-		: RecentPlayer(RecentPlayer)
+	FFriendRecentPlayerItem(TSharedPtr< FOnlineRecentPlayer > InRecentPlayer)
+		: RecentPlayer(InRecentPlayer)
 	{ }
 
 public:
diff --git a/Engine/Source/Developer/FriendsAndChat/Private/Models/FriendViewModel.cpp b/Engine/Source/Developer/FriendsAndChat/Private/Models/FriendViewModel.cpp
--- a/Engine/Source/Developer/FriendsAndChat/Private/Models/FriendViewModel.cpp
+++ b/Engine/Source/Developer/FriendsAndChat/Private/Models/FriendViewModel.cpp
@@ -274,9 +274,9 @@ private:
 	}
 
 	FFriendViewModelImpl(
-		const TSharedRef<IFriendItem>& FriendItem
+		const TSharedRef<IFriendItem>& InFriendItem
 		)
-		: FriendItem(FriendItem)
+		: FriendItem(InFriendItem)
 	{
 	}
 
diff --git a/Engine/Source/Developer/FriendsAndChat/Private/Models/FriendsStatusViewModel.cpp b/Engine/Source/Developer/FriendsAndChat/Private/Models/FriendsStatusViewModel.cpp
--- a/Engine/Source/Developer/FriendsAndChat/Private/Models/FriendsStatusViewModel.cpp
+++ b/Engine/Source/Developer/FriendsAndChat/Private/Models/FriendsStatusViewModel.cpp
@@ -35,9 +35,9 @@ public:
 
 private:
 	FFriendsStatusViewModelImpl(
-		const TSharedRef<FFriendsAndChatManager>& FriendsAndChatManager
+		const TSharedRef<FFriendsAndChatManager>& InFriendsAndChatManager
 		)
-		: FriendsAndChatManager(FriendsAndChatManager)
+		: FriendsAndChatManager(InFriendsAndChatManager)
 	{
 		OnlineStateArray.Add(FOnlineState(true, NSLOCTEXT("OnlineState", "OnlineState_Online", "Online"), EOnlinePresenceState::Online));
 		OnlineStateArray.Add(FOnlineState(false, NSLOCTEXT("OnlineState", "OnlineState_Offline", "Offline"), EOnlinePresenceState::Offline));
diff --git a/Engine/Source/Developer/FriendsAndChat/Private/Models/FriendsUserSettingsViewModel.cpp b/Engine/Source/Developer/FriendsAndChat/Private/Models/FriendsUserSettingsViewModel.cpp
--- a/Engine/Source/Developer/FriendsAndChat/Private/Models/FriendsUserSettingsViewModel.cpp
+++ b/Engine/Source/Developer/FriendsAndChat/Private/Models/FriendsUserSettingsViewModel.cpp
@@ -12,9 +12,9 @@ public:
 
 private:
 	FFriendsUserSettingsViewModelImpl(
-		const TSharedRef<FFriendsAndChatManager>& FriendsAndChatManager
+		const TSharedRef<FFriendsAndChatManager>& InFriendsAndChatManager
 		)
-		: FriendsAndChatManager(FriendsAndChatManager)
+		: FriendsAndChatManager(InFriendsAndChatManager)
 		, bShowNotifications(true)
 	{
 	}
diff --git a/Engine/Source/Developer/FriendsAndChat/Private/Models/FriendsUserViewModel.cpp b/Engine/Source/Developer/FriendsAndChat/Private/Models/FriendsUserViewModel.cpp
--- a/Engine/Source/Developer/FriendsAndChat/Private/Models/FriendsUserViewModel.cpp
+++ b/Engine/Source/Developer/FriendsAndChat/Private/Models/FriendsUserViewModel.cpp
@@ -56,9 +56,9 @@ private:
 	}
 
 	FFriendsUserViewModelImpl(
-		const TSharedRef<FFriendsAndChatManager>& FriendsAndChatManager
+		const TSharedRef<FFriendsAndChatManager>& InFriendsAndChatManager
 		)
-		: FriendsAndChatManager(FriendsAndChatManager)
+		: FriendsAndChatManager(InFriendsAndChatManager)
 	{
 	}
 
diff --git a/Engine/Source/Developer/FriendsAndChat/Private/Models/FriendsViewModel.cpp b/Engine/Source/Developer/FriendsAndChat/Private/Models/FriendsViewModel.cpp
--- a/Engine/Source/Developer/FriendsAndChat/Private/Models/FriendsViewModel.cpp
+++ b/Engine/Source/Developer/FriendsAndChat/Private/Models/FriendsViewModel.cpp
@@ -103,13 +103,13 @@ private:
 	}
 
 	FFriendsViewModelImpl(
-		const TSharedRef<FFriendsAndChatManager>& FriendsAndChatManager,
-		const TSharedRef<IClanRepository>& ClanRepository,
-		const TSharedRef<IFriendListFactory>& FriendsListFactory
+		const TSharedRef<FFriendsAndChatManager>& InFriendsAndChatManager,
+		const TSharedRef<IClanRepository>& InClanRepository,
+		const TSharedRef<IFriendListFactory>& InFriendsListFactory
 		)
-		: FriendsAndChatManager(FriendsAndChatManager)
-		, ClanRepository(ClanRepository)
-		, FriendsListFactory(FriendsListFactory)
+		: FriendsAndChatManager(InFriendsAndChatManager)
+		, ClanRepository(InClanRepository)
+		, FriendsListFactory(InFriendsListFactory)
 		, bIsPerformingAction(false)
 	{
 	}
diff --git a/Engine/Source/Developer/FriendsAndChat/Private/UI/Widgets/SClanItem.cpp b/Engine/Source/Developer/FriendsAndChat/Private/UI/Widgets/SClanItem.cpp
--- a/Engine/Source/Developer/FriendsAndChat/Private/UI/Widgets/SClanItem.cpp
+++ b/Engine/Source/Developer/FriendsAndChat/Private/UI/Widgets/SClanItem.cpp
@@ -11,10 +11,10 @@ class SClanItemImpl : public SClanItem
 {
 public:
 
-	void Construct(const FArguments& InArgs, const TSharedRef<class IClanInfo>& ClanInfo)
+	void Construct(const FArguments& InArgs, const TSharedRef<class IClanInfo>& InClanInfo)
 	{
 		FriendStyle = *InArgs._FriendStyle;
-		this->ClanInfo = ClanInfo;
+		ClanInfo = InClanInfo;
 
 		FFriendsAndChatComboButtonStyle ActionButtonStyle;
 		ActionButtonStyle.ComboButtonStyle = &FriendStyle.ActionComboButtonStyle;
diff --git a/Engine/Source/Developer/FriendsAndChat/Private/UI/Widgets/SClanList.cpp b/Engine/Source/Developer/FriendsAndChat/Private/UI/Widgets/SClanList.cpp
--- a/Engine/Source/Developer/FriendsAndChat/Private/UI/Widgets/SClanList.cpp
+++ b/Engine/Source/Developer/FriendsAndChat/Private/UI/Widgets/SClanList.cpp
@@ -14,10 +14,10 @@ class SClanListImpl : public SClanList
 {
 public:
 
-	void Construct(const FArguments& InArgs, const TSharedRef<FClanListViewModel>& ViewModel)
+	void Construct(const FArguments& InArgs, const TSharedRef<FClanListViewModel>& InViewModel)
 	{
 		FriendStyle = *InArgs._FriendStyle;
-		this->ViewModel = ViewModel;
+		ViewModel = InViewModel;
 		//ViewModel->OnFriendsListUpdated().AddSP(this, &SFriendsListImpl::RefreshFriendsList);
 		MenuMethod = InArgs._Method;
 
diff --git a/Engine/Source/Developer/FriendsAndChat/Private/UI/Widgets/SFriendRequest.cpp b/Engine/Source/Developer/FriendsAndChat/Private/UI/Widgets/SFriendRequest.cpp
--- a/Engine/Source/Developer/FriendsAndChat/Private/UI/Widgets/SFriendRequest.cpp
+++ b/Engine/Source/Developer/FriendsAndChat/Private/UI/Widgets/SFriendRequest.cpp
@@ -13,10 +13,10 @@ class SFriendRequestImpl : public SFriendRequest
 {
 public:
 
-	void Construct(const FArguments& InArgs, const TSharedRef<FFriendsViewModel>& ViewModel)
+	void Construct(const FArguments& InArgs, const TSharedRef<FFriendsViewModel>& InViewModel)
 	{
 		FriendStyle = *InArgs._FriendStyle;
-		this->ViewModel = ViewModel;
+		ViewModel = InViewModel;
 		// Set up titles
 		const FText SearchStringLabel = LOCTEXT( "FriendList_SearchLabel", "[Enter Display Name]" );
 
diff --git a/Engine/Source/Developer/FriendsAndChat/Private/UI/Widgets/SFriendsContainer.cpp b/Engine/Source/Developer/FriendsAndChat/Private/UI/Widgets/SFriendsContainer.cpp
--- a/Engine/Source/Developer/FriendsAndChat/Private/UI/Widgets/SFriendsContainer.cpp
+++ b/Engine/Source/Developer/FriendsAndChat/Private/UI/Widgets/SFriendsContainer.cpp
@@ -26,9 +26,9 @@ class SFriendsContainerImpl : public SFriendsContainer
 {
 public:
 
-	void Construct(const FArguments& InArgs, const TSharedRef<FFriendsViewModel>& ViewModel)
+	void Construct(const FArguments& InArgs, const TSharedRef<FFriendsViewModel>& InViewModel)
 	{
-		this->ViewModel = ViewModel;
+		ViewModel = InViewModel;
 		FriendStyle = *InArgs._FriendStyle;
 
 		for (int32 ListIndex = 0; ListIndex < EFriendsDisplayLists::MAX_None; ListIndex++)
diff --git a/Engine/Source/Developer/FriendsAndChat/Private/UI/Widgets/SFriendsList.cpp b/Engine/Source/Developer/FriendsAndChat/Private/UI/Widgets/SFriendsList.cpp
--- a/Engine/Source/Developer/FriendsAndChat/Private/UI/Widgets/SFriendsList.cpp
+++ b/Engine/Source/Developer/FriendsAndChat/Private/UI/Widgets/SFriendsList.cpp
@@ -16,10 +16,10 @@ class SFriendsListImpl : public SFriendsList
 {
 public:
 
-	void Construct(const FArguments& InArgs, const TSharedRef<FFriendListViewModel>& ViewModel)
+	void Construct(const FArguments& InArgs, const TSharedRef<FFriendListViewModel>& InViewModel)
 	{
 		FriendStyle = *InArgs._FriendStyle;
-		this->ViewModel = ViewModel;
+		ViewModel = InViewModel;
 		ViewModel->OnFriendsListUpdated().AddSP(this, &SFriendsListImpl::RefreshFriendsList);
 		MenuMethod = InArgs._Method;
 
diff --git a/Engine/Source/Developer/FriendsAndChat/Private/UI/Widgets/SFriendsStatus.cpp b/Engine/Source/Developer/FriendsAndChat/Private/UI/Widgets/SFriendsStatus.cpp
--- a/Engine/Source/Developer/FriendsAndChat/Private/UI/Widgets/SFriendsStatus.cpp
+++ b/Engine/Source/Developer/FriendsAndChat/Private/UI/Widgets/SFriendsStatus.cpp
@@ -14,12 +14,12 @@ class SFriendsStatusImpl : public SFriendsStatus
 {
 public:
 
-	void Construct(const FArguments& InArgs, const TSharedRef<FFriendsStatusViewModel>& ViewModel)
+	void Construct(const FArguments& InArgs, const TSharedRef<FFriendsStatusViewModel>& InViewModel)
 	{
 		FriendStyle = *InArgs._FriendStyle;
-		this->ViewModel = ViewModel;
+		ViewModel = InViewModel;
 
-		FFriendsStatusViewModel* ViewModelPtr = &ViewModel.Get();
+		FFriendsStatusViewModel* ViewModelPtr = &InViewModel.Get();
 
 		const TArray<FFriendsStatusViewModel::FOnlineState>& StatusOptions = ViewModelPtr->GetStatusList();
 		SFriendsAndChatCombo::FItemsArray ComboMenuItems;
diff --git a/Engine/Source/Developer/FriendsAndChat/Private/UI/Widgets/SFriendsStatusCombo.cpp b/Engine/Source/Developer/FriendsAndChat/Private/UI/Widgets/SFriendsStatusCombo.cpp
--- a/Engine/Source/Developer/FriendsAndChat/Private/UI/Widgets/SFriendsStatusCombo.cpp
+++ b/Engine/Source/Developer/FriendsAndChat/Private/UI/Widgets/SFriendsStatusCombo.cpp
@@ -16,13 +16,13 @@ class SFriendsStatusComboImpl : public SFriendsStatusCombo
 {
 public:
 
-	void Construct(const FArguments& InArgs, const TSharedRef<FFriendsStatusViewModel>& ViewModel, const TSharedRef<FFriendsUserViewModel>& UserViewModel)
+	void Construct(const FArguments& InArgs, const TSharedRef<FFriendsStatusViewModel>& InViewModel, const TSharedRef<FFriendsUserViewModel>& InUserViewModel)
 	{
 		FriendStyle = *InArgs._FriendStyle;
-		this->ViewModel = ViewModel;
-		this->UserViewModel = UserViewModel;
+		ViewModel = InViewModel;
+		UserViewModel = InUserViewModel;
 
-		FFriendsStatusViewModel* ViewModelPtr = &ViewModel.Get();
+		FFriendsStatusViewModel* ViewModelPtr = &InViewModel.Get();
 
 		const TArray<FFriendsStatusViewModel::FOnlineState>& StatusOptions = ViewModelPtr->GetStatusList();
 		SFriendsAndChatCombo::FItemsArray ComboMenuItems;
diff --git a/Engine/Source/Developer/FriendsAndChat/Private/UI/Widgets/SFriendsUserSettings.cpp b/Engine/Source/Developer/FriendsAndChat/Private/UI/Widgets/SFriendsUserSettings.cpp
--- a/Engine/Source/Developer/FriendsAndChat/Private/UI/Widgets/SFriendsUserSettings.cpp
+++ b/Engine/Source/Developer/FriendsAndChat/Private/UI/Widgets/SFriendsUserSettings.cpp
@@ -13,10 +13,10 @@ class SFriendsUserSettingsImpl : public SFriendsUserSettings
 {
 public:
 
-	void Construct(const FArguments& InArgs, const TSharedRef<FFriendsUserSettingsViewModel>& ViewModel)
+	void Construct(const FArguments& InArgs, const TSharedRef<FFriendsUserSettingsViewModel>& InViewModel)
 	{
 		FriendStyle = *InArgs._FriendStyle;
-		this->ViewModel = ViewModel;
+		ViewModel = InViewModel;
 
 		SUserWidget::Construct(SUserWidget::FArguments()
 		[
diff --git a/Engine/Source/Developer/FriendsAndChat/Private/UI/Widgets/SInviteItem.cpp b/Engine/Source/Developer/FriendsAndChat/Private/UI/Widgets/SInviteItem.cpp
--- a/Engine/Source/Developer/FriendsAndChat/Private/UI/Widgets/SInviteItem.cpp
+++ b/Engine/Source/Developer/FriendsAndChat/Private/UI/Widgets/SInviteItem.cpp
@@ -12,10 +12,10 @@ class SInviteItemImpl : public SInviteItem
 {
 public:
 
-	void Construct(const FArguments& InArgs, const TSharedRef<FFriendViewModel>& ViewModel)
+	void Construct(const FArguments& InArgs, const TSharedRef<FFriendViewModel>& InViewModel)
 	{
 		FriendStyle = *InArgs._FriendStyle;
-		this->ViewModel = ViewModel;
+		ViewModel = InViewModel;
 		FFriendViewModel* ViewModelPtr = this->ViewModel.Get();
 
 		SUserWidget::Construct(SUserWidget::FArguments()
diff --git a/Engine/Source/Developer/FriendsAndChat/Private/UI/Widgets/SInvitesList.cpp b/Engine/Source/Developer/FriendsAndChat/Private/UI/Widgets/SInvitesList.cpp
--- a/Engine/Source/Developer/FriendsAndChat/Private/UI/Widgets/SInvitesList.cpp
+++ b/Engine/Source/Developer/FriendsAndChat/Private/UI/Widgets/SInvitesList.cpp
@@ -15,10 +15,10 @@ class SInvitesListImpl : public SInvitesList
 {
 public:
 
-	void Construct(const FArguments& InArgs, const TSharedRef<FFriendListViewModel>& ViewModel)
+	void Construct(const FArguments& InArgs, const TSharedRef<FFriendListViewModel>& InViewModel)
 	{
 		FriendStyle = *InArgs._FriendStyle;
-		this->ViewModel = ViewModel;
+		ViewModel = InViewModel;
 		
 		ViewModel->OnFriendsListUpdated().AddSP(this, &SInvitesListImpl::RefreshFriendsList);
 
