Fix UnrealLauncher et al shadow variables - batch 2

diff --git a/Engine/Source/Developer/FriendsAndChat/Private/FriendsAndChatManager.cpp b/Engine/Source/Developer/FriendsAndChat/Private/FriendsAndChatManager.cpp
--- a/Engine/Source/Developer/FriendsAndChat/Private/FriendsAndChatManager.cpp
+++ b/Engine/Source/Developer/FriendsAndChat/Private/FriendsAndChatManager.cpp
@@ -248,9 +248,9 @@ void FFriendsAndChatManager::AddApplicationViewModel(const FString ClientID, TSh
 	ApplicationViewModels.Add(ClientID, InApplicationViewModel);
 }
 
-void FFriendsAndChatManager::SetUserSettings(const FFriendsAndChatSettings& UserSettings)
+void FFriendsAndChatManager::SetUserSettings(const FFriendsAndChatSettings& InUserSettings)
 {
-	this->UserSettings = UserSettings;
+	UserSettings = InUserSettings;
 }
 
 void FFriendsAndChatManager::SetAnalyticsProvider(const TSharedPtr<IAnalyticsProvider>& AnalyticsProvider)
@@ -1813,12 +1813,12 @@ void FFriendsAndChatManager::OnPresenceReceived(const FUniqueNetId& UserId, cons
 	// Check for out of date invites
 	if (OnlineSub != nullptr && OnlineIdentity.IsValid())
 	{
-		TSharedPtr<FUniqueNetId> UserId = OnlineIdentity->GetUniquePlayerId(0);
-		if (UserId.IsValid())
+		TSharedPtr<FUniqueNetId> UserId0 = OnlineIdentity->GetUniquePlayerId(0);
+		if (UserId0.IsValid())
 		{
 			TSharedPtr<FOnlineUserPresence> CurrentPresence;
 			TArray<FString> InvitesToRemove;
-			OnlineSub->GetPresenceInterface()->GetCachedPresence(*UserId, CurrentPresence);
+			OnlineSub->GetPresenceInterface()->GetCachedPresence(*UserId0, CurrentPresence);
 			for (auto It = PendingGameInvitesList.CreateConstIterator(); It; ++It)
 			{
 				FString CurrentSessionID = CurrentPresence->SessionId.IsValid() ? CurrentPresence->SessionId->ToString() : TEXT("");
