Fix compile error

diff --git a/Engine/Source/Developer/FriendsAndChat/Private/UI/Widgets/SFriendsStatusCombo.cpp b/Engine/Source/Developer/FriendsAndChat/Private/UI/Widgets/SFriendsStatusCombo.cpp
--- a/Engine/Source/Developer/FriendsAndChat/Private/UI/Widgets/SFriendsStatusCombo.cpp
+++ b/Engine/Source/Developer/FriendsAndChat/Private/UI/Widgets/SFriendsStatusCombo.cpp
@@ -55,7 +55,7 @@ public:
 					.BorderImage(FCoreStyle::Get().GetBrush("NoBorder"))
 					.Padding(FMargin(0, 0, 25, 0))
 					[
-						SNew(SFriendUserHeader, UserViewModel)
+						SNew(SFriendUserHeader, InUserViewModel)
 						.ShowFriendName(true)
 						.FriendStyle(&FriendStyle)
 					]
