[INTEGRATE] merge UE4-LauncheDev > UE4 CL 2581600
Add the chat entry style

diff --git a/Engine/Source/Developer/FriendsAndChat/Private/Layers/UI/SChatEntryWidget.cpp b/Engine/Source/Developer/FriendsAndChat/Private/Layers/UI/SChatEntryWidget.cpp
--- a/Engine/Source/Developer/FriendsAndChat/Private/Layers/UI/SChatEntryWidget.cpp
+++ b/Engine/Source/Developer/FriendsAndChat/Private/Layers/UI/SChatEntryWidget.cpp
@@ -50,7 +50,7 @@ public:
 			+SHorizontalBox::Slot()
 			[
 				SAssignNew(ChatTextBox, SMultiLineEditableTextBox)
-				.Style(InArgs._Style)
+				.Style(&FriendStyle.FriendsChatStyle.ChatEntryTextStyle)
 				.ClearKeyboardFocusOnCommit(false)
 				.OnTextCommitted(this, &SChatEntryWidgetImpl::HandleChatEntered)
 				.OnKeyDownHandler(this, &SChatEntryWidgetImpl::HandleChatKeydown)
