More PIE fixes to allow us to chat from in-game. Also fixing up some warnings

diff --git a/Engine/Source/Developer/FriendsAndChat/Private/Layers/DataAccess/FriendsAndChatManager.cpp b/Engine/Source/Developer/FriendsAndChat/Private/Layers/DataAccess/FriendsAndChatManager.cpp
--- a/Engine/Source/Developer/FriendsAndChat/Private/Layers/DataAccess/FriendsAndChatManager.cpp
+++ b/Engine/Source/Developer/FriendsAndChat/Private/Layers/DataAccess/FriendsAndChatManager.cpp
@@ -62,7 +62,7 @@ void FFriendsAndChatManager::Initialize()
 /* IFriendsAndChatManager interface
  *****************************************************************************/
 
-void FFriendsAndChatManager::Login(IOnlineSubsystem* InOnlineSub, bool bInIsGame)
+void FFriendsAndChatManager::Login(IOnlineSubsystem* InOnlineSub, bool bInIsGame, int32 LocalUserID)
 {
 	// Clear existing data
 	Logout();
@@ -147,7 +147,7 @@ void FFriendsAndChatManager::Login(IOnlineSubsystem* InOnlineSub, bool bInIsGame
 			SetState(EFriendsAndManagerState::RequestFriendsListRefresh);
 			RequestRecentPlayersListRefresh();
 
-			MessageManager->LogIn();
+			MessageManager->LogIn(OnlineSub, LocalUserID);
 			MessageManager->OnChatPublicRoomJoined().AddSP(this, &FFriendsAndChatManager::OnChatPublicRoomJoined);
 			for (auto RoomName : ChatRoomstoJoin)
 			{
@@ -405,8 +405,8 @@ TSharedPtr< SWidget > FFriendsAndChatManager::GenerateChatWidget(const FFriendsA
 
 	if (WindowPtr.IsValid())
 	{
-		WindowPtr->SetOnWindowActivated(FOnWindowActivated::CreateSP(ChatWidget.Get(), &SChatWindow::HandleWindowActivated));
-		WindowPtr->SetOnWindowDeactivated(FOnWindowDeactivated::CreateSP(ChatWidget.Get(), &SChatWindow::HandleWindowDeactivated));
+		WindowPtr->GetOnWindowActivatedEvent().AddSP(ChatWidget.Get(), &SChatWindow::HandleWindowActivated);
+		WindowPtr->GetOnWindowDeactivatedEvent().AddSP(ChatWidget.Get(), &SChatWindow::HandleWindowDeactivated);
 	}
 
 	return ChatWidget;
diff --git a/Engine/Source/Developer/FriendsAndChat/Private/Layers/DataAccess/FriendsAndChatManager.h b/Engine/Source/Developer/FriendsAndChat/Private/Layers/DataAccess/FriendsAndChatManager.h
--- a/Engine/Source/Developer/FriendsAndChat/Private/Layers/DataAccess/FriendsAndChatManager.h
+++ b/Engine/Source/Developer/FriendsAndChat/Private/Layers/DataAccess/FriendsAndChatManager.h
@@ -126,7 +126,7 @@ public:
 
 	// IFriendsAndChatManager
 	virtual void Logout() override;
-	virtual void Login(IOnlineSubsystem* InOnlineSub, bool bInIsGame = false) override;
+	virtual void Login(IOnlineSubsystem* InOnlineSub, bool bInIsGame = false, int32 LocalUserID = 0) override;
 	virtual bool IsLoggedIn() override;
 	virtual void SetOnline() override;
 	virtual void SetAway() override;
diff --git a/Engine/Source/Developer/FriendsAndChat/Private/Layers/DataAccess/FriendsMessageManager.cpp b/Engine/Source/Developer/FriendsAndChat/Private/Layers/DataAccess/FriendsMessageManager.cpp
--- a/Engine/Source/Developer/FriendsAndChat/Private/Layers/DataAccess/FriendsMessageManager.cpp
+++ b/Engine/Source/Developer/FriendsAndChat/Private/Layers/DataAccess/FriendsMessageManager.cpp
@@ -22,14 +22,18 @@ class FFriendsMessageManagerImpl
 {
 public:
 
-	virtual void LogIn() override
+	virtual void LogIn(IOnlineSubsystem* InOnlineSub, int32 LocalUserNum) override
 	{
-		Initialize();
+		// Clear existing data
+		LogOut();
 
+		OnlineSub = InOnlineSub;
+
+		Initialize();
 		if (OnlineSub != nullptr &&
 			OnlineSub->GetIdentityInterface().IsValid())
 		{
-			LoggedInUser = OnlineSub->GetIdentityInterface()->GetUniquePlayerId(0);
+			LoggedInUser = OnlineSub->GetIdentityInterface()->GetUniquePlayerId(LocalUserNum);
 		}
 
 		for (auto RoomName : RoomJoins)
@@ -187,17 +191,12 @@ private:
 
 	void Initialize()
 	{
-		// Clear existing data
-		LogOut();
-
 		GlobalMessagesCount = 0;
 		WhisperMessagesCount = 0;
 		GameMessagesCount = 0;
 		PartyMessagesCount = 0;
 		ReceivedMessages.Empty();
 
-		OnlineSub = IOnlineSubsystem::Get( TEXT( "MCP" ) );
-
 		if (OnlineSub != nullptr)
 		{
 			IOnlineChatPtr ChatInterface = OnlineSub->GetChatInterface();
diff --git a/Engine/Source/Developer/FriendsAndChat/Private/Layers/DataAccess/FriendsMessageManager.h b/Engine/Source/Developer/FriendsAndChat/Private/Layers/DataAccess/FriendsMessageManager.h
--- a/Engine/Source/Developer/FriendsAndChat/Private/Layers/DataAccess/FriendsMessageManager.h
+++ b/Engine/Source/Developer/FriendsAndChat/Private/Layers/DataAccess/FriendsMessageManager.h
@@ -32,7 +32,7 @@ public:
 	/** Destructor. */
 	virtual ~FFriendsMessageManager( ) {};
 
-	virtual void LogIn() = 0;
+	virtual void LogIn(IOnlineSubsystem* InOnlineSub, int32 LocalUserNum = 0) = 0;
 	virtual void LogOut() = 0;
 	virtual const TArray<TSharedRef<FFriendChatMessage> >& GetMessages() const = 0;
 	virtual void JoinPublicRoom(const FString& RoomName) = 0;
diff --git a/Engine/Source/Developer/FriendsAndChat/Public/Interfaces/IFriendsAndChatManager.h b/Engine/Source/Developer/FriendsAndChat/Public/Interfaces/IFriendsAndChatManager.h
--- a/Engine/Source/Developer/FriendsAndChat/Public/Interfaces/IFriendsAndChatManager.h
+++ b/Engine/Source/Developer/FriendsAndChat/Public/Interfaces/IFriendsAndChatManager.h
@@ -109,7 +109,7 @@ public:
 	virtual void Logout() = 0;
 
 	/** Log in and start checking for Friends. */
-	virtual void Login(IOnlineSubsystem* InOnlineSub, bool bInIsGame = false) = 0;
+	virtual void Login(IOnlineSubsystem* InOnlineSub, bool bInIsGame = false, int32 LocalUserID = 0) = 0;
 
 	/** Is the chat manager logged in. */
 	virtual bool IsLoggedIn() = 0;
