Fix build warning about deprecated push menu

diff --git a/Engine/Source/Developer/FriendsAndChat/Private/Layers/UI/SChatWindow.cpp b/Engine/Source/Developer/FriendsAndChat/Private/Layers/UI/SChatWindow.cpp
--- a/Engine/Source/Developer/FriendsAndChat/Private/Layers/UI/SChatWindow.cpp
+++ b/Engine/Source/Developer/FriendsAndChat/Private/Layers/UI/SChatWindow.cpp
@@ -292,7 +292,14 @@ private:
 			FText Username = FText::FromString(*UsernameString);
 			const TSharedRef<FFriendViewModel> FriendViewModel = ViewModel->GetFriendViewModel(*UniqueIDString, Username).ToSharedRef();
 			TSharedRef<SWidget> Widget = SNew(SFriendActions, FriendViewModel).FriendStyle(&FriendStyle);
-			FSlateApplication::Get().PushMenu(AsShared(), Widget, FSlateApplication::Get().GetCursorPos(), FPopupTransitionEffect(FPopupTransitionEffect::ContextMenu));
+
+			EntryMenu = FSlateApplication::Get().PushMenu(
+				SharedThis(this),
+				FWidgetPath(),
+				Widget,
+				FSlateApplication::Get().GetCursorPos(),
+				FPopupTransitionEffect(FPopupTransitionEffect::ContextMenu)
+				);
 		}
 	}
 
@@ -452,6 +459,9 @@ private:
 
 	FSlateHyperlinkRun::FOnClick OnHyperlinkClicked;
 
+	/** Reference to owner of the current pop-up. */
+	TWeakPtr<IMenu> EntryMenu;
+
 	static const float CHAT_HINT_UPDATE_THROTTLE;
 };
 
