Fix Linux Mono build - remove unused expression result

diff --git a/Engine/Source/Developer/FriendsAndChat/Private/Layers/UI/SChatWindow.cpp b/Engine/Source/Developer/FriendsAndChat/Private/Layers/UI/SChatWindow.cpp
--- a/Engine/Source/Developer/FriendsAndChat/Private/Layers/UI/SChatWindow.cpp
+++ b/Engine/Source/Developer/FriendsAndChat/Private/Layers/UI/SChatWindow.cpp
@@ -293,7 +293,7 @@ private:
 			const TSharedRef<FFriendViewModel> FriendViewModel = ViewModel->GetFriendViewModel(*UniqueIDString, Username).ToSharedRef();
 			TSharedRef<SWidget> Widget = SNew(SFriendActions, FriendViewModel).FriendStyle(&FriendStyle);
 
-			EntryMenu = FSlateApplication::Get().PushMenu(
+			FSlateApplication::Get().PushMenu(
 				SharedThis(this),
 				FWidgetPath(),
 				Widget,
@@ -459,9 +459,6 @@ private:
 
 	FSlateHyperlinkRun::FOnClick OnHyperlinkClicked;
 
-	/** Reference to owner of the current pop-up. */
-	TWeakPtr<IMenu> EntryMenu;
-
 	static const float CHAT_HINT_UPDATE_THROTTLE;
 };
 
