[INTEGRATE] merge UE4-LauncherDev > UE4 CL 2583896
Remove set chat tips when we remove slash markup

diff --git a/Engine/Source/Developer/FriendsAndChat/Private/Layers/Domain/FriendsChatMarkupService.cpp b/Engine/Source/Developer/FriendsAndChat/Private/Layers/Domain/FriendsChatMarkupService.cpp
--- a/Engine/Source/Developer/FriendsAndChat/Private/Layers/Domain/FriendsChatMarkupService.cpp
+++ b/Engine/Source/Developer/FriendsAndChat/Private/Layers/Domain/FriendsChatMarkupService.cpp
@@ -22,6 +22,14 @@ struct FProcessedInput
 	bool FoundMatch;
 	TArray<TSharedPtr<FFriendViewModel> > ValidFriends;
 	TArray<TSharedPtr<FFriendViewModel> > MatchedFriends;
+
+	void Clear()
+	{
+		ValidFriends.Empty();
+		MatchedFriends.Empty();
+		NeedsTip = false;
+		FoundMatch = false;
+	}
 };
 
 class FChatTip : public IChatTip
@@ -438,6 +446,9 @@ private:
 			}
 			return ProcessedInput;
 		}
+
+		// Text is no longer markup - clear any existing tips
+		ProcessedInput->Clear();
 		return nullptr;
 	}
 
