Merging using UE4-To-UE4-LauncherDev

2584013 #friends+chat, analytics
- fix for not sending private chat counts

2584233 #friends+chat
- fix for analytics events not being flushed on shutdown
- fix for analytics events not being sent due to nullptr provider

diff --git a/Engine/Source/Developer/FriendsAndChat/Private/Layers/DataAccess/FriendsAndChatManager.cpp b/Engine/Source/Developer/FriendsAndChat/Private/Layers/DataAccess/FriendsAndChatManager.cpp
--- a/Engine/Source/Developer/FriendsAndChat/Private/Layers/DataAccess/FriendsAndChatManager.cpp
+++ b/Engine/Source/Developer/FriendsAndChat/Private/Layers/DataAccess/FriendsAndChatManager.cpp
@@ -46,7 +46,6 @@ FFriendsAndChatManager::FFriendsAndChatManager()
 
 FFriendsAndChatManager::~FFriendsAndChatManager( )
 {
-	Analytics.FlushChatStats();
 }
 
 void FFriendsAndChatManager::Initialize()
@@ -165,6 +164,9 @@ void FFriendsAndChatManager::Login(IOnlineSubsystem* InOnlineSub, bool bInIsGame
 
 void FFriendsAndChatManager::Logout()
 {
+	// flush before removing the analytics provider
+	Analytics.FlushChatStats(); 
+
 	if (OnlineSub != nullptr)
 	{
 		if (OnlineSub->GetFriendsInterface().IsValid())
@@ -215,9 +217,6 @@ void FFriendsAndChatManager::Logout()
 	OnlineIdentity = nullptr;
 	FriendsInterface = nullptr;
 	OnlineSub = nullptr;
-	// Prevent a shared pointer issue with stale analytics providers
-    // @todo do this as a failsafe, but unsafe now until FTickerObjectBase issues are resolved
-	SetAnalyticsProvider(nullptr);
 
 	if ( UpdateFriendsTickerDelegate.IsBound() )
 	{
@@ -1991,7 +1990,7 @@ void FFriendsAndChatAnalytics::FlushChatStats()
 					if (ChatCounts.Num())
 					{
 						TArray<FAnalyticsEventAttribute> Attributes;
-						for (const auto& Pair : ChannelChatCounts)
+						for (const auto& Pair : ChatCounts)
 						{
 							Attributes.Empty(3);
 							Attributes.Emplace(TEXT("Name"), Pair.Key);
