Fix clang error; unclear if properly.

(Merging CL 2587348 from LauncherDev to main).

diff --git a/Engine/Source/Developer/FriendsAndChat/Private/Layers/UI/SChatWindow.cpp b/Engine/Source/Developer/FriendsAndChat/Private/Layers/UI/SChatWindow.cpp
--- a/Engine/Source/Developer/FriendsAndChat/Private/Layers/UI/SChatWindow.cpp
+++ b/Engine/Source/Developer/FriendsAndChat/Private/Layers/UI/SChatWindow.cpp
@@ -246,7 +246,7 @@ private:
 		}
 	
 		TArray<TSharedRef<FChatItemViewModel>> Messages = ViewModel->GetMessages();
-		for (LastMessageIndex; LastMessageIndex < Messages.Num(); ++LastMessageIndex)
+		for (LastMessageIndex = 0; LastMessageIndex < Messages.Num(); ++LastMessageIndex)
 		{
 			TSharedRef<FChatItemViewModel> ChatItem = Messages[LastMessageIndex];
 
