Merging

//depot/UE4-LauncherDev/Engine/Source/Developer/FriendsAndChat/Private/Layers/UI/SChatWindow.cpp

to D:\\Perforce\\Epic\\UE4\\Engine\\Source\\Developer\\FriendsAndChat\\Private\\Layers\\UI\\SChatWindow.cpp

diff --git a/Engine/Source/Developer/FriendsAndChat/Private/Layers/UI/SChatWindow.cpp b/Engine/Source/Developer/FriendsAndChat/Private/Layers/UI/SChatWindow.cpp
--- a/Engine/Source/Developer/FriendsAndChat/Private/Layers/UI/SChatWindow.cpp
+++ b/Engine/Source/Developer/FriendsAndChat/Private/Layers/UI/SChatWindow.cpp
@@ -246,7 +246,7 @@ private:
 		}
 	
 		TArray<TSharedRef<FChatItemViewModel>> Messages = ViewModel->GetMessages();
-		for (LastMessageIndex = 0; LastMessageIndex < Messages.Num(); ++LastMessageIndex)
+		for (;LastMessageIndex < Messages.Num(); ++LastMessageIndex)
 		{
 			TSharedRef<FChatItemViewModel> ChatItem = Messages[LastMessageIndex];
 
