[INTEGRATE] merge UE4-LauncherDev > UE4 CL 2588345
Fix Friends appearing as offline in chat window. View Model was not getting returned.

diff --git a/Engine/Source/Developer/FriendsAndChat/Private/Layers/DataAccess/FriendsAndChatManager.cpp b/Engine/Source/Developer/FriendsAndChat/Private/Layers/DataAccess/FriendsAndChatManager.cpp
--- a/Engine/Source/Developer/FriendsAndChat/Private/Layers/DataAccess/FriendsAndChatManager.cpp
+++ b/Engine/Source/Developer/FriendsAndChat/Private/Layers/DataAccess/FriendsAndChatManager.cpp
@@ -1410,7 +1410,7 @@ TSharedPtr<FFriendViewModel> FFriendsAndChatManager::GetFriendViewModel(const TS
 	}
 	if(FoundFriend.IsValid())
 	{
-		FriendViewModelFactory->Create(FoundFriend.ToSharedRef());
+		return FriendViewModelFactory->Create(FoundFriend.ToSharedRef());
 	}
 	return nullptr;
 }
