[INTEGRATE] merge UE4-LauncherDev > UE4
Friends & Chat changes up to 2588351

diff --git a/Engine/Source/Developer/FriendsAndChat/Private/FriendsAndChatStyle.cpp b/Engine/Source/Developer/FriendsAndChat/Private/FriendsAndChatStyle.cpp
--- a/Engine/Source/Developer/FriendsAndChat/Private/FriendsAndChatStyle.cpp
+++ b/Engine/Source/Developer/FriendsAndChat/Private/FriendsAndChatStyle.cpp
@@ -106,19 +106,19 @@ TSharedRef< FSlateStyleSet > FFriendsAndChatModuleStyle::Create(FFriendsAndChatS
 
 	// Name Style
 	const FTextBlockStyle GlobalChatFont = FTextBlockStyle(DefaultText)
-		.SetFont(FriendStyle.FriendsNormalFontStyle.FriendsFontSmallBold)
-		.SetColorAndOpacity(FriendStyle.FriendsChatStyle.DefaultChatColor);
+		.SetFont(FriendStyle.FriendsNormalFontStyle.FriendsFontNormalBold)
+		.SetColorAndOpacity(FriendStyle.FriendsChatStyle.GlobalChatColor);
 
 	const FTextBlockStyle GameChatFont = FTextBlockStyle(DefaultText)
-		.SetFont(FriendStyle.FriendsNormalFontStyle.FriendsFontSmallBold)
+		.SetFont(FriendStyle.FriendsNormalFontStyle.FriendsFontNormalBold)
 		.SetColorAndOpacity(FriendStyle.FriendsChatStyle.GameChatColor);
 
 	const FTextBlockStyle PartyChatFont = FTextBlockStyle(DefaultText)
-		.SetFont(FriendStyle.FriendsNormalFontStyle.FriendsFontSmallBold)
+		.SetFont(FriendStyle.FriendsNormalFontStyle.FriendsFontNormalBold)
 		.SetColorAndOpacity(FriendStyle.FriendsChatStyle.PartyChatColor);
 
 	const FTextBlockStyle WhisperChatFont = FTextBlockStyle(DefaultText)
-		.SetFont(FriendStyle.FriendsNormalFontStyle.FriendsFontSmallBold)
+		.SetFont(FriendStyle.FriendsNormalFontStyle.FriendsFontNormalBold)
 		.SetColorAndOpacity(FriendStyle.FriendsChatStyle.WhisplerChatColor);
 
 	const FButtonStyle UserNameButton = FButtonStyle()
diff --git a/Engine/Source/Developer/FriendsAndChat/Private/FriendsChatStyle.cpp b/Engine/Source/Developer/FriendsAndChat/Private/FriendsChatStyle.cpp
--- a/Engine/Source/Developer/FriendsAndChat/Private/FriendsChatStyle.cpp
+++ b/Engine/Source/Developer/FriendsAndChat/Private/FriendsChatStyle.cpp
@@ -40,6 +40,12 @@ FFriendsChatStyle& FFriendsChatStyle::SetDefaultChatColor(const FLinearColor& In
 	return *this;
 }
 
+FFriendsChatStyle& FFriendsChatStyle::SetGlobalChatColor(const FLinearColor& InFontColor)
+{
+	GlobalChatColor = InFontColor;
+	return *this;
+}
+
 FFriendsChatStyle& FFriendsChatStyle::SetWhisplerChatColor(const FLinearColor& InFontColor)
 {
 	WhisplerChatColor = InFontColor;
diff --git a/Engine/Source/Developer/FriendsAndChat/Private/Layers/DataAccess/FriendsAndChatManager.cpp b/Engine/Source/Developer/FriendsAndChat/Private/Layers/DataAccess/FriendsAndChatManager.cpp
--- a/Engine/Source/Developer/FriendsAndChat/Private/Layers/DataAccess/FriendsAndChatManager.cpp
+++ b/Engine/Source/Developer/FriendsAndChat/Private/Layers/DataAccess/FriendsAndChatManager.cpp
@@ -1087,36 +1087,33 @@ void FFriendsAndChatManager::OnQueryRecentPlayersComplete(const FUniqueNetId& Us
 	}
 	else
 	{
-		if (bWasSuccessful)
+		RecentPlayersList.Empty();
+		TArray< TSharedRef<FOnlineRecentPlayer> > Players;
+		if (FriendsInterface->GetRecentPlayers(UserId, FString(), Players))
 		{
-			RecentPlayersList.Empty();
-			TArray< TSharedRef<FOnlineRecentPlayer> > Players;
-			if (FriendsInterface->GetRecentPlayers(UserId, FString(), Players))
+			for (const auto& RecentPlayer : Players)
 			{
-				for (const auto& RecentPlayer : Players)
+				if (RecentPlayer->GetDisplayName().IsEmpty())
 				{
-					if (RecentPlayer->GetDisplayName().IsEmpty())
-					{
-						QueryUserIds.Add(RecentPlayer->GetUserId());
-					}
-					else
-					{
-						RecentPlayersList.Add(MakeShareable(new FFriendRecentPlayerItem(RecentPlayer)));
-					}
-				}
-
-				if (QueryUserIds.Num())
-				{
-					check(OnlineSub != nullptr && OnlineSub->GetUserInterface().IsValid());
-					IOnlineUserPtr UserInterface = OnlineSub->GetUserInterface();
-					UserInterface->QueryUserInfo(LocalControllerIndex, QueryUserIds);
-					bFoundAllIds = false;
+					QueryUserIds.Add(RecentPlayer->GetUserId());
 				}
 				else
 				{
-					OnFriendsListUpdated().Broadcast();
+					RecentPlayersList.Add(MakeShareable(new FFriendRecentPlayerItem(RecentPlayer)));
 				}
 			}
+
+			if (QueryUserIds.Num())
+			{
+				check(OnlineSub != nullptr && OnlineSub->GetUserInterface().IsValid());
+				IOnlineUserPtr UserInterface = OnlineSub->GetUserInterface();
+				UserInterface->QueryUserInfo(LocalControllerIndex, QueryUserIds);
+				bFoundAllIds = false;
+			}
+			else
+			{
+				OnFriendsListUpdated().Broadcast();
+			}
 		}
 
 		if (bFoundAllIds)
diff --git a/Engine/Source/Developer/FriendsAndChat/Private/Layers/Domain/FriendsChatMarkupService.cpp b/Engine/Source/Developer/FriendsAndChat/Private/Layers/Domain/FriendsChatMarkupService.cpp
--- a/Engine/Source/Developer/FriendsAndChat/Private/Layers/Domain/FriendsChatMarkupService.cpp
+++ b/Engine/Source/Developer/FriendsAndChat/Private/Layers/Domain/FriendsChatMarkupService.cpp
@@ -16,6 +16,12 @@ class FFriendsChatMarkupServiceFactoryImpl;
 
 struct FProcessedInput
 {
+	FProcessedInput()
+		: ChatChannel(EChatMessageType::Invalid)
+		, NeedsTip(false)
+		, FoundMatch(false)
+	{}
+
 	EChatMessageType::Type ChatChannel;
 	FString Message;
 	bool NeedsTip;
diff --git a/Engine/Source/Developer/FriendsAndChat/Private/Layers/Presentation/ChatViewModel.cpp b/Engine/Source/Developer/FriendsAndChat/Private/Layers/Presentation/ChatViewModel.cpp
--- a/Engine/Source/Developer/FriendsAndChat/Private/Layers/Presentation/ChatViewModel.cpp
+++ b/Engine/Source/Developer/FriendsAndChat/Private/Layers/Presentation/ChatViewModel.cpp
@@ -157,7 +157,7 @@ public:
 	virtual bool SendMessage(const FText NewMessage, const FText PlainText) override
 	{
 		bool bSuccess = true;
-		if (!MarkupService->ValidateSlashMarkup(NewMessage.ToString(), PlainText.ToString()))
+		if (!AllowMarkup() || !MarkupService->ValidateSlashMarkup(NewMessage.ToString(), PlainText.ToString()))
 		{
 			if(!PlainText.IsEmptyOrWhitespace())
 			{
@@ -459,10 +459,13 @@ public:
 
 	virtual void ValidateChatInput(const FText Message, const FText PlainText)
 	{
-		MarkupService->SetInputText(Message.ToString(), PlainText.ToString());
-		if(!Message.IsEmpty())
+		if (AllowMarkup())
 		{
-			ChatDisplayService->ChatEntered();
+			MarkupService->SetInputText(Message.ToString(), PlainText.ToString());
+			if (!Message.IsEmpty())
+			{
+				ChatDisplayService->ChatEntered();
+			}
 		}
 	}
 
@@ -501,6 +504,21 @@ public:
 // 		}
 	}
 
+	virtual bool AllowMarkup() override
+	{
+		return true;
+	}
+	
+	virtual bool MultiChat() override
+	{
+		return true;
+	}
+
+	virtual void SetFocus() override
+	{
+		ChatDisplayService->SetFocus();
+	}
+
 	virtual bool IsActive() const override
 	{
 		return bIsActive;
@@ -536,7 +554,7 @@ public:
 
 	virtual FReply HandleChatKeyEntry(const FKeyEvent& KeyEvent) override
 	{
-		return MarkupService->HandleChatKeyEntry(KeyEvent);
+		return AllowMarkup() ? MarkupService->HandleChatKeyEntry(KeyEvent) : FReply::Unhandled();
 	}
 
 	DECLARE_DERIVED_EVENT(FChatViewModelImpl, FChatViewModel::FChatListSetFocus, FChatListSetFocus);
@@ -702,6 +720,17 @@ private:
 		}
 		return Changed;
 	}
+
+	virtual bool AllowMarkup() override
+	{
+		return false;
+	}
+
+	virtual bool MultiChat() override
+	{
+		return false;
+	}
+	
 	// End FChatViewModelImpl 
 
 
diff --git a/Engine/Source/Developer/FriendsAndChat/Private/Layers/Presentation/ChatViewModel.h b/Engine/Source/Developer/FriendsAndChat/Private/Layers/Presentation/ChatViewModel.h
--- a/Engine/Source/Developer/FriendsAndChat/Private/Layers/Presentation/ChatViewModel.h
+++ b/Engine/Source/Developer/FriendsAndChat/Private/Layers/Presentation/ChatViewModel.h
@@ -82,6 +82,9 @@ public:
 	virtual void SetIsActive(bool IsActive) = 0;
 	virtual void RefreshMessages() = 0;
 	virtual void SetInParty(bool bInPartySetting) = 0;
+	virtual bool AllowMarkup() = 0;
+	virtual bool MultiChat() = 0;
+	virtual void SetFocus() = 0;
 
 	// Settings
 	virtual void SetChatSettingsService(TSharedPtr<IChatSettingsService> ChatSettingsService) = 0;
diff --git a/Engine/Source/Developer/FriendsAndChat/Private/Layers/UI/SAutoTextScroller.cpp b/Engine/Source/Developer/FriendsAndChat/Private/Layers/UI/SAutoTextScroller.cpp
new file mode 100644
--- /dev/null
+++ b/Engine/Source/Developer/FriendsAndChat/Private/Layers/UI/SAutoTextScroller.cpp
@@ -0,0 +1,123 @@
+// Copyright 1998-2015 Epic Games, Inc. All Rights Reserved.
+
+#include "FriendsAndChatPrivatePCH.h"
+#include "SAutoTextScroller.h"
+
+namespace EScrollerState
+{
+	enum Type
+	{
+		Start = 0,
+		Scrolling,
+		End
+	};
+}
+
+class SAutoTextScrollerImpl : public SAutoTextScroller
+{
+public:
+
+	void Construct(const FArguments& InArgs)
+	{
+		FontStyle = *InArgs._FontStyle;
+
+		SAssignNew(ExternalScrollbar, SScrollBar);
+		SUserWidget::Construct(SUserWidget::FArguments()
+		[
+			SAssignNew(ScrollBox, SScrollBox)
+			.ScrollBarVisibility(EVisibility::Collapsed)
+			.Orientation(Orient_Horizontal)
+			.ExternalScrollbar(ExternalScrollbar)
+			+ SScrollBox::Slot()
+			[
+				SNew(STextBlock)
+				.Font(FontStyle.FriendsFontSmallBold)
+				.ColorAndOpacity(FontStyle.DefaultDullFontColor)
+				.Text(InArgs._Text)
+			]
+		]);
+
+		TimeElapsed = 0.f;
+		ScrollOffset = 0;
+		ScrollerState = EScrollerState::Start;
+	}
+
+	virtual void Tick(const FGeometry& AllottedGeometry, const double InCurrentTime, const float InDeltaTime) override
+	{
+		const float ContentSize = ScrollBox->GetDesiredSize().X;
+		const float Speed = 20.0f;
+		const float StartDelay = 2.0f;
+		const float EndDelay = 2.0f;
+
+		TimeElapsed += InDeltaTime;
+
+		switch (ScrollerState)
+		{
+			case EScrollerState::Start:
+			{
+				auto PinnedActiveTimerHandle = ActiveTimerHandle.Pin();
+				if (PinnedActiveTimerHandle.IsValid())
+				{
+					UnRegisterActiveTimer(PinnedActiveTimerHandle.ToSharedRef());
+				}
+				ActiveTimerHandle = RegisterActiveTimer(StartDelay, FWidgetActiveTimerDelegate::CreateSP(this, &SAutoTextScrollerImpl::UpdateScrollerState));
+			}break;
+			case EScrollerState::Scrolling:
+			{
+				ScrollOffset += Speed * InDeltaTime;
+				if (ExternalScrollbar->DistanceFromBottom()==0.0f)
+				{
+					ScrollOffset = ContentSize;
+					TimeElapsed = 0.0f;
+					ScrollerState = EScrollerState::End;
+				}
+				
+			}break;
+			case EScrollerState::End:
+			{
+				auto PinnedActiveTimerHandle = ActiveTimerHandle.Pin();
+				if (PinnedActiveTimerHandle.IsValid())
+				{
+					UnRegisterActiveTimer(PinnedActiveTimerHandle.ToSharedRef());
+				}
+				ActiveTimerHandle = RegisterActiveTimer(EndDelay, FWidgetActiveTimerDelegate::CreateSP(this, &SAutoTextScrollerImpl::UpdateScrollerState));
+			}break;
+		};
+		ScrollBox->SetScrollOffset(ScrollOffset);
+	}
+
+	EActiveTimerReturnType UpdateScrollerState(double InCurrentTime, float InDeltaTime)
+	{
+		if (ScrollerState == EScrollerState::End)
+		{
+			TimeElapsed = 0.0f;
+			ScrollOffset = 0.0f;
+			ScrollerState = EScrollerState::Start;
+			ScrollBox->SetScrollOffset(ScrollOffset);
+		}
+		else if (ScrollerState == EScrollerState::Start)
+		{
+			TimeElapsed = 0.0f;
+			ScrollOffset = 0.0f;
+			ScrollerState = EScrollerState::Scrolling;
+			ScrollBox->SetScrollOffset(ScrollOffset);
+		}
+		return EActiveTimerReturnType::Continue;
+	}
+
+private:
+
+	/** Holds the style to use when making the widget. */
+	FFriendsFontStyle FontStyle;
+	TSharedPtr<SScrollBox> ScrollBox;
+	TWeakPtr<FActiveTimerHandle> ActiveTimerHandle;
+	TSharedPtr<SScrollBar> ExternalScrollbar;
+	EScrollerState::Type ScrollerState;
+	float TimeElapsed;
+	float ScrollOffset;
+};
+
+TSharedRef<SAutoTextScroller> SAutoTextScroller::New()
+{
+	return MakeShareable(new SAutoTextScrollerImpl());
+}
diff --git a/Engine/Source/Developer/FriendsAndChat/Private/Layers/UI/SAutoTextScroller.h b/Engine/Source/Developer/FriendsAndChat/Private/Layers/UI/SAutoTextScroller.h
new file mode 100644
--- /dev/null
+++ b/Engine/Source/Developer/FriendsAndChat/Private/Layers/UI/SAutoTextScroller.h
@@ -0,0 +1,20 @@
+// Copyright 1998-2015 Epic Games, Inc. All Rights Reserved.
+
+#pragma once
+
+class SAutoTextScroller : public SUserWidget
+{
+public:
+	SLATE_USER_ARGS(SAutoTextScroller)
+	{ }
+	SLATE_ARGUMENT(const FFriendsFontStyle*, FontStyle)
+	SLATE_ATTRIBUTE(FText, Text)
+	SLATE_END_ARGS()
+
+	/**
+	 * Constructs the widget.
+	 *
+	 * @param InArgs The Slate argument list.
+	 */
+	virtual void Construct( const FArguments& InArgs ) = 0;
+};
diff --git a/Engine/Source/Developer/FriendsAndChat/Private/Layers/UI/SChatEntryWidget.cpp b/Engine/Source/Developer/FriendsAndChat/Private/Layers/UI/SChatEntryWidget.cpp
--- a/Engine/Source/Developer/FriendsAndChat/Private/Layers/UI/SChatEntryWidget.cpp
+++ b/Engine/Source/Developer/FriendsAndChat/Private/Layers/UI/SChatEntryWidget.cpp
@@ -33,6 +33,7 @@ public:
 			.AutoWidth()
 			[
 			 	SNew(SBox)
+				.Visibility(ViewModel->AllowMarkup() ? EVisibility::Visible : EVisibility::Collapsed)
 			 	.WidthOverride(HightOverride)
 			 	[
 					SNew(SButton)
diff --git a/Engine/Source/Developer/FriendsAndChat/Private/Layers/UI/SChatWindow.cpp b/Engine/Source/Developer/FriendsAndChat/Private/Layers/UI/SChatWindow.cpp
--- a/Engine/Source/Developer/FriendsAndChat/Private/Layers/UI/SChatWindow.cpp
+++ b/Engine/Source/Developer/FriendsAndChat/Private/Layers/UI/SChatWindow.cpp
@@ -34,8 +34,10 @@ public:
 			);
 
 		OnHyperlinkClicked = FSlateHyperlinkRun::FOnClick::CreateSP(SharedThis(this), &SChatWindowImpl::HandleNameClicked);
+		OnChannelHyperlinkClicked = FSlateHyperlinkRun::FOnClick::CreateSP(SharedThis(this), &SChatWindowImpl::HandleChannelClicked);
 
 		RichTextMarshaller->AppendInlineDecorator(FHyperlinkDecorator::Create(TEXT("UserName"), OnHyperlinkClicked));
+		RichTextMarshaller->AppendInlineDecorator(FHyperlinkDecorator::Create(TEXT("Channel"), OnChannelHyperlinkClicked));
 		RichTextMarshaller->AppendInlineDecorator(FImageDecorator::Create(TEXT("img"), &FFriendsAndChatModuleStyle::Get()));
 
 		TimeTransparency = 0.0f;
@@ -130,7 +132,7 @@ public:
 									.Style(&FriendStyle.FriendsChatStyle.ChatEntryTextStyle)
 									.TextStyle(&FriendStyle.FriendsChatStyle.TextStyle)
 								 	.FriendStyle(&FriendStyle)
-									.Marshaller(RichTextMarshaller.ToSharedRef())
+									.Marshaller(RichTextMarshaller)
 									.HintText(InArgs._ActivationHintText)
 								]
 								+SVerticalBox::Slot()
@@ -138,7 +140,7 @@ public:
 								[
 									SNew(SBorder)
 									.Padding(FMargin(10, 10))
-									.VAlign(VAlign_Center)									
+									.VAlign(VAlign_Center)
 									.Visibility(this, &SChatWindowImpl::GetLostConnectionVisibility)
 									.BorderImage(&FriendStyle.FriendsChatStyle.ChatContainerBackground)
 									.BorderBackgroundColor(FLinearColor(FColor(255, 255, 255, 128)))
@@ -162,7 +164,7 @@ public:
 
 	virtual void HandleWindowActivated() override
 	{
-		FSlateApplication::Get().SetKeyboardFocus(ChatTextBox, EFocusCause::WindowActivate);
+		ViewModel->SetFocus();
 		ViewModel->SetIsActive(true);
 	}
 
@@ -303,6 +305,17 @@ private:
 		}
 	}
 
+	void HandleChannelClicked(const FSlateHyperlinkRun::FMetadata& Metadata)
+	{
+		const FString* ChannelString = Metadata.Find(TEXT("Command"));
+
+		if (ChannelString)
+		{
+			ViewModel->SetOutgoingMessageChannel(EChatMessageType::EnumFromString(*ChannelString));
+			ViewModel->SetFocus();
+		}
+	}
+
 	FText FormatListMessageText(TSharedRef<FChatItemViewModel> ChatItem, bool bGrouped)
 	{
 		const FText& TextToSet = ChatItem->GetMessage();
@@ -314,18 +327,19 @@ private:
 		{
 		case EChatMessageType::Global: 
 			HyperlinkStyle = TEXT("UserNameTextStyle.GlobalHyperlink"); 
-			Args.Add(TEXT("ChannelImage"), FText::FromString(TEXT("GlobalChatIcon")));
+			Args.Add(TEXT("RoomName"), FText::FromString(TEXT("[g]")));
 			break;
 		case EChatMessageType::Whisper: 
 			HyperlinkStyle = TEXT("UserNameTextStyle.Whisperlink"); 
-			Args.Add(TEXT("ChannelImage"), FText::FromString(TEXT("WhisperChatIcon")));
+			Args.Add(TEXT("RoomName"), FText::FromString(TEXT("[w]")));
 			break;
 		case EChatMessageType::Game: 
 			HyperlinkStyle = TEXT("UserNameTextStyle.GameHyperlink"); 
-			Args.Add(TEXT("ChannelImage"), FText::FromString(TEXT("PartyChatIcon")));
+			Args.Add(TEXT("RoomName"), FText::FromString(TEXT("[p]")));
 			break;
 		}
 
+		Args.Add(TEXT("Channel"), FText::FromString(EChatMessageType::ShortcutString(ChatItem->GetMessageType())));
 		Args.Add(TEXT("RecipientName"), ChatItem->GetRecipientName());
 		if(ChatItem->GetRecipientID().IsValid())
 		{
@@ -349,20 +363,41 @@ private:
 		}
 		else
 		{
-			if (ChatItem->IsFromSelf())
+			if (!ViewModel->MultiChat())
 			{
-				if(ChatItem->GetMessageType() == EChatMessageType::Whisper)
+				if (ChatItem->IsFromSelf())
 				{
-					ListMessage = FText::Format(NSLOCTEXT("FChatItemViewModel", "DisplayNameAndMessage", "<img src=\"{ChannelImage}\"/>{SenderName} to <a id=\"UserName\" uid=\"{RecipientID}\" Username=\"{RecipientName}\" style=\"{NameStyle}\">{RecipientName}</> {MessageTime}: {Message}"), Args);
+					ListMessage = FText::Format(NSLOCTEXT("FChatItemViewModel", "DisplayNameAndMessage", "{MessageTime} {SenderName} : {Message}"), Args);
 				}
 				else
 				{
-					ListMessage = FText::Format(NSLOCTEXT("FChatItemViewModel", "DisplayNameAndMessage", "<img src=\"{ChannelImage}\"/>{SenderName} {MessageTime}: {Message}"), Args);
+					ListMessage = FText::Format(NSLOCTEXT("FChatItemViewModel", "DisplayNameAndMessage", "{MessageTime} <a id=\"UserName\" uid=\"{SenderID}\" Username=\"{SenderName}\" style=\"{NameStyle}\">{SenderName}</> : {Message}"), Args);
 				}
 			}
 			else
 			{
-				ListMessage = FText::Format(NSLOCTEXT("FChatItemViewModel", "DisplayNameAndMessage", "<img src=\"{ChannelImage}\"/> <a id=\"UserName\" uid=\"{SenderID}\" Username=\"{SenderName}\" style=\"{NameStyle}\">{SenderName}</> {MessageTime}: {Message}"), Args);
+				if (ChatItem->GetMessageType() == EChatMessageType::Whisper)
+				{
+					if (ChatItem->IsFromSelf())
+					{
+						ListMessage = FText::Format(NSLOCTEXT("FChatItemViewModel", "DisplayNameAndMessage", "{MessageTime} <a id=\"Channel\" Command=\"{Channel}\" style=\"{NameStyle}\">{RoomName}</> to <a id=\"UserName\" uid=\"{RecipientID}\" Username=\"{RecipientName}\" style=\"{NameStyle}\">{RecipientName}</> : {Message}"), Args);
+					}
+					else
+					{
+						ListMessage = FText::Format(NSLOCTEXT("FChatItemViewModel", "DisplayNameAndMessage", "{MessageTime} <a id=\"Channel\" Command=\"{Channel}\" style=\"{NameStyle}\">{RoomName}</> from <a id=\"UserName\" uid=\"{SenderID}\" Username=\"{SenderName}\" style=\"{NameStyle}\">{SenderName}</> : {Message}"), Args);
+					}
+				}
+				else
+				{
+					if (ChatItem->IsFromSelf())
+					{
+						ListMessage = FText::Format(NSLOCTEXT("FChatItemViewModel", "DisplayNameAndMessage", "{MessageTime} <a id=\"Channel\" Command=\"{Channel}\" style=\"{NameStyle}\">{RoomName}</> {SenderName} : {Message}"), Args);
+					}
+					else
+					{
+						ListMessage = FText::Format(NSLOCTEXT("FChatItemViewModel", "DisplayNameAndMessage", "{MessageTime} <a id=\"Channel\" Command=\"{Channel}\" style=\"{NameStyle}\">{RoomName}</> <a id=\"UserName\" uid=\"{SenderID}\" Username=\"{SenderName}\" style=\"{NameStyle}\">{SenderName}</> : {Message}"), Args);
+					}
+				}
 			}
 		}
 		return ListMessage;
@@ -458,6 +493,7 @@ private:
 	FCurveHandle FadeCurve;
 
 	FSlateHyperlinkRun::FOnClick OnHyperlinkClicked;
+	FSlateHyperlinkRun::FOnClick OnChannelHyperlinkClicked;
 
 	static const float CHAT_HINT_UPDATE_THROTTLE;
 };
diff --git a/Engine/Source/Developer/FriendsAndChat/Private/Layers/UI/SFriendItem.cpp b/Engine/Source/Developer/FriendsAndChat/Private/Layers/UI/SFriendItem.cpp
--- a/Engine/Source/Developer/FriendsAndChat/Private/Layers/UI/SFriendItem.cpp
+++ b/Engine/Source/Developer/FriendsAndChat/Private/Layers/UI/SFriendItem.cpp
@@ -6,6 +6,7 @@
 #include "SFriendsToolTip.h"
 #include "SFriendsList.h"
 #include "SFriendsAndChatCombo.h"
+#include "SAutoTextScroller.h"
 
 #define LOCTEXT_NAMESPACE "SFriendItem"
 
@@ -90,9 +91,8 @@ public:
 								.VAlign(VAlign_Top)
 								.AutoHeight()
 								[
-									SNew(STextBlock)
-									.Font(FriendFontStyle.FriendsFontSmallBold)
-									.ColorAndOpacity(FriendFontStyle.DefaultDullFontColor)
+									SNew(SAutoTextScroller)
+									.FontStyle(&FriendFontStyle)
 									.Text(ViewModelPtr, &FFriendViewModel::GetFriendLocation)
 								]
 							]
diff --git a/Engine/Source/Developer/FriendsAndChat/Private/Layers/UI/SFriendsContainer.cpp b/Engine/Source/Developer/FriendsAndChat/Private/Layers/UI/SFriendsContainer.cpp
--- a/Engine/Source/Developer/FriendsAndChat/Private/Layers/UI/SFriendsContainer.cpp
+++ b/Engine/Source/Developer/FriendsAndChat/Private/Layers/UI/SFriendsContainer.cpp
@@ -371,7 +371,7 @@ private:
 							SNew(SBox)
 							.Visibility(this, &SFriendsContainerImpl::AddFriendBoxVisibility)
 							.VAlign(VAlign_Center)
-							.HAlign(HAlign_Left)
+							.HAlign(HAlign_Fill)
 							.Padding(0)
 							[
 								SAssignNew(FriendNameTextBox, SEditableTextBox)
diff --git a/Engine/Source/Developer/FriendsAndChat/Public/FriendsChatStyle.h b/Engine/Source/Developer/FriendsAndChat/Public/FriendsChatStyle.h
--- a/Engine/Source/Developer/FriendsAndChat/Public/FriendsChatStyle.h
+++ b/Engine/Source/Developer/FriendsAndChat/Public/FriendsChatStyle.h
@@ -66,6 +66,10 @@ struct FRIENDSANDCHAT_API FFriendsChatStyle
 	FLinearColor WhisplerChatColor;
 	FFriendsChatStyle& SetWhisplerChatColor(const FLinearColor& InFontColor);
 
+	UPROPERTY(EditAnywhere, Category = Appearance)
+	FLinearColor GlobalChatColor;
+	FFriendsChatStyle& SetGlobalChatColor(const FLinearColor& InFontColor);
+
 	UPROPERTY(EditAnywhere, Category=Appearance)
 	FLinearColor GameChatColor;
 	FFriendsChatStyle& SetGameChatColor(const FLinearColor& InFontColor);
