Buildwatcher - Fix shadow variables

diff --git a/Engine/Source/Developer/FriendsAndChat/Private/Models/FriendItem.h b/Engine/Source/Developer/FriendsAndChat/Private/Models/FriendItem.h
--- a/Engine/Source/Developer/FriendsAndChat/Private/Models/FriendItem.h
+++ b/Engine/Source/Developer/FriendsAndChat/Private/Models/FriendItem.h
@@ -21,8 +21,8 @@ public:
 	 * @param InOnlineUser The online user.
 	 * @param InListType The list type.
 	 */
-	FFriendItem(TSharedPtr< FOnlineFriend > InOnlineFriend, TSharedPtr< FOnlineUser > InOnlineUser, EFriendsDisplayLists::Type InListType, const TSharedRef<class FFriendsAndChatManager>& FriendsAndChatManager)
-		: FriendsAndChatManager(FriendsAndChatManager)
+	FFriendItem(TSharedPtr< FOnlineFriend > InOnlineFriend, TSharedPtr< FOnlineUser > InOnlineUser, EFriendsDisplayLists::Type InListType, const TSharedRef<class FFriendsAndChatManager>& InFriendsAndChatManager)
+		: FriendsAndChatManager(InFriendsAndChatManager)
 		, bIsUpdated(true)
 		, GroupName(TEXT(""))
 		, OnlineFriend( InOnlineFriend )
