Merging

//depot/UE4-LauncherDev/...

to //depot/UE4/...

diff --git a/Engine/Source/Developer/FriendsAndChat/Private/FriendsAndChatStyle.cpp b/Engine/Source/Developer/FriendsAndChat/Private/FriendsAndChatStyle.cpp
--- a/Engine/Source/Developer/FriendsAndChat/Private/FriendsAndChatStyle.cpp
+++ b/Engine/Source/Developer/FriendsAndChat/Private/FriendsAndChatStyle.cpp
@@ -2,6 +2,7 @@
 
 #include "FriendsAndChatPrivatePCH.h"
 
+#include "FriendsFontStyleService.h"
 const FName FFriendsAndChatStyle::TypeName( TEXT("FFriendsAndChatStyle") );
 
 FFriendsAndChatStyle& FFriendsAndChatStyle::SetNormalFriendsFontStyle(const FFriendsFontStyle& FontStyle)
@@ -74,6 +75,7 @@ const FFriendsAndChatStyle& FFriendsAndChatStyle::GetDefault()
 	Module style set
 */
 TSharedPtr< FSlateStyleSet > FFriendsAndChatModuleStyle::FriendsAndChatModuleStyleInstance = NULL;
+TSharedPtr< FFriendsFontStyleService > FFriendsAndChatModuleStyle::FriendsFontStyleService = NULL;
 
 void FFriendsAndChatModuleStyle::Initialize(FFriendsAndChatStyle FriendStyle)
 {
@@ -81,6 +83,9 @@ void FFriendsAndChatModuleStyle::Initialize(FFriendsAndChatStyle FriendStyle)
 	{
 		FriendsAndChatModuleStyleInstance = Create(FriendStyle);
 		FSlateStyleRegistry::RegisterSlateStyle( *FriendsAndChatModuleStyleInstance );
+
+		FriendsFontStyleService = FFriendsFontStyleServiceFactory::Create();
+		FriendsFontStyleService->SetFontStyles(FriendStyle.FriendsNormalFontStyle);
 	}
 }
 
@@ -97,6 +102,15 @@ FName FFriendsAndChatModuleStyle::GetStyleSetName()
 	return StyleSetName;
 }
 
+TSharedPtr<FFriendsFontStyleService> FFriendsAndChatModuleStyle::GetStyleService()
+{
+	if (!FriendsFontStyleService.IsValid())
+	{
+		FriendsFontStyleService = FFriendsFontStyleServiceFactory::Create();
+	}
+	return FriendsFontStyleService;
+}
+
 TSharedRef< FSlateStyleSet > FFriendsAndChatModuleStyle::Create(FFriendsAndChatStyle FriendStyle)
 {
 	TSharedRef< FSlateStyleSet > Style = MakeShareable(new FSlateStyleSet("FriendsAndChatStyle"));
diff --git a/Engine/Source/Developer/FriendsAndChat/Private/FriendsFontStyleService.cpp b/Engine/Source/Developer/FriendsAndChat/Private/FriendsFontStyleService.cpp
new file mode 100644
--- /dev/null
+++ b/Engine/Source/Developer/FriendsAndChat/Private/FriendsFontStyleService.cpp
@@ -0,0 +1,55 @@
+// Copyright 1998-2015 Epic Games, Inc. All Rights Reserved.
+
+#include "FriendsAndChatPrivatePCH.h"
+#include "FriendsFontStyleService.h"
+
+class FFriendsFontStyleServiceImpl
+	: public FFriendsFontStyleService
+{
+public:
+
+	virtual void SetFontStyles(const FFriendsFontStyle& InFriendsNormalFontStyle) override
+	{
+		FriendsNormalFontStyle = InFriendsNormalFontStyle;
+	}
+
+	virtual void SetUserFontSize(EChatFontType::Type InFontSize) override
+	{
+		FontSize = InFontSize;
+	}
+
+	virtual FSlateFontInfo GetNormalFont() override
+	{
+		if (FontSize == EChatFontType::Large)
+		{
+			return FriendsNormalFontStyle.FriendsFontLarge;
+		}
+		return FriendsNormalFontStyle.FriendsFontNormal;
+	}
+
+	virtual FSlateFontInfo GetNormalBoldFont() override
+	{
+		if (FontSize == EChatFontType::Large)
+		{
+			return FriendsNormalFontStyle.FriendsFontLargeBold;
+		}
+		return FriendsNormalFontStyle.FriendsFontNormalBold;
+	}
+
+private:
+	FFriendsFontStyleServiceImpl()
+	{
+		FontSize = EChatFontType::Standard;
+	}
+
+	FFriendsFontStyle FriendsNormalFontStyle;
+	EChatFontType::Type FontSize;
+
+	friend FFriendsFontStyleServiceFactory;
+};
+
+TSharedRef< FFriendsFontStyleService > FFriendsFontStyleServiceFactory::Create()
+{
+	TSharedRef< FFriendsFontStyleServiceImpl > StyleService(new FFriendsFontStyleServiceImpl());
+	return StyleService;
+}
\ No newline at end of file
diff --git a/Engine/Source/Developer/FriendsAndChat/Private/Layers/Presentation/ChatSettingsService.cpp b/Engine/Source/Developer/FriendsAndChat/Private/Layers/Presentation/ChatSettingsService.cpp
--- a/Engine/Source/Developer/FriendsAndChat/Private/Layers/Presentation/ChatSettingsService.cpp
+++ b/Engine/Source/Developer/FriendsAndChat/Private/Layers/Presentation/ChatSettingsService.cpp
@@ -177,6 +177,26 @@ private:
 		}
 	}
 
+	virtual uint8 GetRadioOptionIndex(EChatSettingsType::Type Option) const override
+	{
+		check(EChatSettingsOptionType::GetOptionType(Option) == EChatSettingsOptionType::RadioBox);
+		switch (Option)
+		{
+		case EChatSettingsType::FontSize:
+		{
+			return ChatSettings.FontSize;
+			break;
+		}
+		case EChatSettingsType::WindowHeight:
+		{
+			return ChatSettings.WindowHeight;
+			break;
+		}
+		default:
+			return 0;
+		}
+	}
+
 	virtual void SetRadioOption(EChatSettingsType::Type Option, uint8 Index) override
 	{
 		check(EChatSettingsOptionType::GetOptionType(Option) == EChatSettingsOptionType::RadioBox);
diff --git a/Engine/Source/Developer/FriendsAndChat/Private/Layers/Presentation/ChatViewModel.cpp b/Engine/Source/Developer/FriendsAndChat/Private/Layers/Presentation/ChatViewModel.cpp
--- a/Engine/Source/Developer/FriendsAndChat/Private/Layers/Presentation/ChatViewModel.cpp
+++ b/Engine/Source/Developer/FriendsAndChat/Private/Layers/Presentation/ChatViewModel.cpp
@@ -1,6 +1,7 @@
 // Copyright 1998-2015 Epic Games, Inc. All Rights Reserved.
 
 #include "FriendsAndChatPrivatePCH.h"
+#include "FriendsFontStyleService.h"
 #include "ChatItemViewModel.h"
 #include "ChatViewModel.h"
 #include "FriendViewModel.h"
@@ -519,6 +520,11 @@ public:
 		ChatDisplayService->SetFocus();
 	}
 
+	virtual float GetWindowOpacity() override
+	{
+		return WindowOpacity;
+	}
+
 	virtual bool IsActive() const override
 	{
 		return bIsActive;
@@ -541,7 +547,12 @@ public:
 	{
 		ChatSettingsService = InChatSettingsService;
 		bAllowFade = ChatSettingsService->GetFlagOption(EChatSettingsType::NeverFadeMessages);
+		WindowOpacity = ChatSettingsService->GetSliderValue(EChatSettingsType::WindowOpacity);
+		uint8 FontSize = ChatSettingsService->GetRadioOptionIndex(EChatSettingsType::FontSize);
+		FFriendsAndChatModuleStyle::GetStyleService()->SetUserFontSize(EChatFontType::Type(FontSize));
 		ChatSettingsService->OnChatSettingStateUpdated().AddSP(this, &FChatViewModelImpl::HandleSettingsChanged);
+		ChatSettingsService->OnChatSettingRadioStateUpdated().AddSP(this, &FChatViewModelImpl::HandleRadioSettingsChanged);
+		ChatSettingsService->OnChatSettingValueUpdated().AddSP(this, &FChatViewModelImpl::HandleValueSettingsChanged);
 	}
 
 	virtual void HandleSettingsChanged(EChatSettingsType::Type OptionType, bool NewState)
@@ -552,6 +563,23 @@ public:
 		}
 	}
 
+	virtual void HandleRadioSettingsChanged(EChatSettingsType::Type OptionType, uint8 Index)
+	{
+		if (OptionType == EChatSettingsType::FontSize)
+		{
+			FFriendsAndChatModuleStyle::GetStyleService()->SetUserFontSize(EChatFontType::Type(Index));
+			OnSettingsUpdated().Broadcast();
+		}
+	}
+
+	virtual void HandleValueSettingsChanged(EChatSettingsType::Type OptionType, float Value)
+	{
+		if (OptionType == EChatSettingsType::WindowOpacity)
+		{
+			WindowOpacity = Value;
+		}
+	}
+
 	virtual FReply HandleChatKeyEntry(const FKeyEvent& KeyEvent) override
 	{
 		return AllowMarkup() ? MarkupService->HandleChatKeyEntry(KeyEvent) : FReply::Unhandled();
@@ -575,6 +603,12 @@ public:
 		return ChatListUpdatedEvent;
 	}
 
+	DECLARE_DERIVED_EVENT(FChatViewModelImpl, FChatViewModel::FChatSettingsUpdated, FChatSettingsUpdated)
+	virtual FChatSettingsUpdated& OnSettingsUpdated() override
+	{
+		return ChatSettingsUpdatedEvent;
+	}
+
 	// End ChatDisplayOptionsViewModel Interface
 
 private:
@@ -611,6 +645,7 @@ protected:
 		, bUseOverrideColor(false)
 		, bAllowGlobalChat(true)
 		, bCaptureFocus(false)
+		, WindowOpacity(1.0f)
 	{
 	}
 
@@ -659,6 +694,7 @@ private:
 	bool bAllowGlobalChat;
 	bool bCaptureFocus;
 	bool bDisplayChatTip;
+	float WindowOpacity;
 
 	FText OutGoingChannelText;
 	// Holds the fade in curve
@@ -666,6 +702,7 @@ private:
 	FChatListSetFocus ChatListSetFocusEvent;
 	FChatTextValidatedEvent ChatTextValidatedEvent;
 	FChatListUpdated ChatListUpdatedEvent;
+	FChatSettingsUpdated ChatSettingsUpdatedEvent;
 
 	EVisibility ChatEntryVisibility;
 	FSlateColor OverrideColor;
diff --git a/Engine/Source/Developer/FriendsAndChat/Private/Layers/Presentation/ChatViewModel.h b/Engine/Source/Developer/FriendsAndChat/Private/Layers/Presentation/ChatViewModel.h
--- a/Engine/Source/Developer/FriendsAndChat/Private/Layers/Presentation/ChatViewModel.h
+++ b/Engine/Source/Developer/FriendsAndChat/Private/Layers/Presentation/ChatViewModel.h
@@ -85,6 +85,7 @@ public:
 	virtual bool AllowMarkup() = 0;
 	virtual bool MultiChat() = 0;
 	virtual void SetFocus() = 0;
+	virtual float GetWindowOpacity() = 0;
 
 	// Settings
 	virtual void SetChatSettingsService(TSharedPtr<IChatSettingsService> ChatSettingsService) = 0;
@@ -97,6 +98,9 @@ public:
 
 	DECLARE_EVENT(FChatViewModel, FChatListUpdated)
 	virtual FChatListUpdated& OnChatListUpdated() = 0;
+
+	DECLARE_EVENT(FChatViewModel, FChatSettingsUpdated)
+	virtual FChatSettingsUpdated& OnSettingsUpdated() = 0;
 };
 
 /**
diff --git a/Engine/Source/Developer/FriendsAndChat/Private/Layers/UI/SChatEntryWidget.cpp b/Engine/Source/Developer/FriendsAndChat/Private/Layers/UI/SChatEntryWidget.cpp
--- a/Engine/Source/Developer/FriendsAndChat/Private/Layers/UI/SChatEntryWidget.cpp
+++ b/Engine/Source/Developer/FriendsAndChat/Private/Layers/UI/SChatEntryWidget.cpp
@@ -1,6 +1,7 @@
 // Copyright 1998-2015 Epic Games, Inc. All Rights Reserved.
 
 #include "FriendsAndChatPrivatePCH.h"
+#include "FriendsFontStyleService.h"
 #include "SChatEntryWidget.h"
 #include "RichTextLayoutMarshaller.h"
 #include "ChatViewModel.h"
@@ -25,6 +26,8 @@ public:
 		FriendStyle = *InArgs._FriendStyle;
 		float HightOverride = FriendStyle.FriendsChatStyle.ChatEntryHeight;
 		MaxChatLength = InArgs._MaxChatLength;
+		HintText = InArgs._HintText;
+		Marshaller = InArgs._Marshaller;
 
 		SUserWidget::Construct(SUserWidget::FArguments()
 		[
@@ -50,19 +53,43 @@ public:
 			]
 			+SHorizontalBox::Slot()
 			[
-				SAssignNew(ChatTextBox, SMultiLineEditableTextBox)
+				SAssignNew(ChatBox, SBox)
+			]
+		]);
+
+		RebuildTextEntry();
+	}
+
+	void RebuildTextEntry() override
+	{
+		FText EnteredText;
+		if (ChatTextBox.IsValid())
+		{
+			EnteredText = ChatTextBox->GetText();
+			ChatTextBox.Reset();
+		}
+		if (!ChatTextBox.IsValid())
+		{
+			SAssignNew(ChatTextBox, SMultiLineEditableTextBox)
 				.Style(&FriendStyle.FriendsChatStyle.ChatEntryTextStyle)
+				.Font(FFriendsAndChatModuleStyle::GetStyleService()->GetNormalFont())
 				.ClearKeyboardFocusOnCommit(false)
 				.OnTextCommitted(this, &SChatEntryWidgetImpl::HandleChatEntered)
 				.OnKeyDownHandler(this, &SChatEntryWidgetImpl::HandleChatKeydown)
-				.HintText(InArgs._HintText)
+				.HintText(HintText)
 				.OnTextChanged(this, &SChatEntryWidgetImpl::HandleChatTextChanged)
 				.IsEnabled(this, &SChatEntryWidgetImpl::IsChatEntryEnabled)
 				.ModiferKeyForNewLine(EModifierKey::Shift)
-				.Marshaller(InArgs._Marshaller)
-				.WrapTextAt(this, &SChatEntryWidgetImpl::GetMessageBoxWrapWidth)
-			]
-		]);
+				.Marshaller(Marshaller)
+				.WrapTextAt(this, &SChatEntryWidgetImpl::GetMessageBoxWrapWidth);
+
+			if (!EnteredText.IsEmpty())
+			{
+				ChatTextBox->SetText(EnteredText);
+			}
+		}
+
+		ChatBox->SetContent(ChatTextBox.ToSharedRef());
 	}
 
 	// Begin SUserWidget
@@ -231,12 +258,18 @@ private:
 	int32 MaxChatLength;
 	/* Holds the current window width (Hack until Rich Text can figure this out) */
 	mutable float WindowWidth;
+	/* Holds the Chat Text */
+	TSharedPtr<SBox> ChatBox;
 	/* Holds the Chat entry box */
 	TSharedPtr<SMultiLineEditableTextBox> ChatTextBox;
 	/** Holds the style to use when making the widget. */
 	FFriendsAndChatStyle FriendStyle;
 	/* Holds the Chat view model */
 	TSharedPtr<FChatViewModel> ViewModel;
+	/* Holds the hint text for the chat box*/
+	TAttribute< FText > HintText;
+	/* Holds the marshaller for the chat box*/
+	TSharedPtr<ITextLayoutMarshaller> Marshaller;
 };
 
 TSharedRef<SChatEntryWidget> SChatEntryWidget::New()
diff --git a/Engine/Source/Developer/FriendsAndChat/Private/Layers/UI/SChatEntryWidget.h b/Engine/Source/Developer/FriendsAndChat/Private/Layers/UI/SChatEntryWidget.h
--- a/Engine/Source/Developer/FriendsAndChat/Private/Layers/UI/SChatEntryWidget.h
+++ b/Engine/Source/Developer/FriendsAndChat/Private/Layers/UI/SChatEntryWidget.h
@@ -13,11 +13,6 @@ public:
 	SLATE_USER_ARGS( SChatEntryWidget )
 		: _MaxChatLength(128)
 		{}
-			/** The styling of the textbox */
-		SLATE_STYLE_ARGUMENT( FEditableTextBoxStyle, Style )
-
-		/** Pointer to a style of the text block, which dictates the font, color, and shadow options. */
-		SLATE_STYLE_ARGUMENT( FTextBlockStyle, TextStyle )
 
 		SLATE_ARGUMENT(const FFriendsAndChatStyle*, FriendStyle)	
 
@@ -39,4 +34,10 @@ public:
 	 * @param InViewModel	The chat view model - used for accessing chat markup etc.
 	 */
 	virtual void Construct(const FArguments& InArgs, const TSharedRef<class FChatViewModel>& InViewModel) = 0;
+
+	/**
+	 * Rebuild the Text Entry Widget
+	 * @param bLargeFont - True if using large font, else standard font
+	 */
+	virtual void RebuildTextEntry() = 0;
 };
diff --git a/Engine/Source/Developer/FriendsAndChat/Private/Layers/UI/SChatWindow.cpp b/Engine/Source/Developer/FriendsAndChat/Private/Layers/UI/SChatWindow.cpp
--- a/Engine/Source/Developer/FriendsAndChat/Private/Layers/UI/SChatWindow.cpp
+++ b/Engine/Source/Developer/FriendsAndChat/Private/Layers/UI/SChatWindow.cpp
@@ -1,6 +1,7 @@
 // Copyright 1998-2015 Epic Games, Inc. All Rights Reserved.
 
 #include "FriendsAndChatPrivatePCH.h"
+#include "FriendsFontStyleService.h"
 #include "SChatWindow.h"
 #include "ChatViewModel.h"
 #include "ChatItemViewModel.h"
@@ -25,6 +26,7 @@ public:
 		ViewModel = InViewModel;
 		FChatViewModel* ViewModelPtr = ViewModel.Get();
 		ViewModel->OnChatListUpdated().AddSP(this, &SChatWindowImpl::RefreshChatList);
+		ViewModel->OnSettingsUpdated().AddSP(this, &SChatWindowImpl::SettingsChanged);
 
 		FFriendsAndChatModuleStyle::Initialize(FriendStyle);
 
@@ -129,8 +131,6 @@ public:
 								+ SVerticalBox::Slot()
 								[
 									SAssignNew(ChatTextBox, SChatEntryWidget, ViewModel.ToSharedRef())
-									.Style(&FriendStyle.FriendsChatStyle.ChatEntryTextStyle)
-									.TextStyle(&FriendStyle.FriendsChatStyle.TextStyle)
 								 	.FriendStyle(&FriendStyle)
 									.Marshaller(RichTextMarshaller)
 									.HintText(InArgs._ActivationHintText)
@@ -230,12 +230,24 @@ private:
 		return true;
 	}
 
-	void RefreshChatList()
+	void SettingsChanged()
 	{
-		if (!RichText.IsValid())
+		CreateChatWidgets();
+		RefreshChatList();
+	}
+
+	void CreateChatWidgets()
 	{
+		if(RichText.IsValid())
+		{
+			RichText.Reset();
+			ChatScrollBox->ClearChildren();
+		}
+		if (!RichText.IsValid())
+		{
 			SAssignNew(RichText, SMultiLineEditableTextBox)
 			.Style(&FriendStyle.FriendsChatStyle.ChatDisplayTextStyle)
+			.Font(FFriendsAndChatModuleStyle::GetStyleService()->GetNormalFont())
 			.WrapTextAt(this, &SChatWindowImpl::GetChatWrapWidth)
 			.Marshaller(RichTextMarshaller.ToSharedRef())
 			.Text(this, &SChatWindowImpl::GetChatText)
@@ -245,8 +257,18 @@ private:
 			[
 				RichText.ToSharedRef()
 			];
+
+			ChatTextBox->RebuildTextEntry();
 		}
-	
+	}
+
+	void RefreshChatList()
+	{
+		if (!RichText.IsValid())
+		{
+			CreateChatWidgets();
+		}
+
 		TArray<TSharedRef<FChatItemViewModel>> Messages = ViewModel->GetMessages();
 		for (;LastMessageIndex < Messages.Num(); ++LastMessageIndex)
 		{
@@ -430,7 +452,7 @@ private:
 
 	FLinearColor GetTimedFadeColor() const
 	{
-		return FLinearColor(1, 1, 1, FadeCurve.GetLerp());
+		return FLinearColor(1, 1, 1, ViewModel->GetWindowOpacity() * FadeCurve.GetLerp());
 	}
 
 	EVisibility GetChatListVisibility() const
diff --git a/Engine/Source/Developer/FriendsAndChat/Public/FriendsAndChatStyle.h b/Engine/Source/Developer/FriendsAndChat/Public/FriendsAndChatStyle.h
--- a/Engine/Source/Developer/FriendsAndChat/Public/FriendsAndChatStyle.h
+++ b/Engine/Source/Developer/FriendsAndChat/Public/FriendsAndChatStyle.h
@@ -107,6 +107,8 @@ public:
 
 	static FName GetStyleSetName();
 
+	static TSharedPtr<class FFriendsFontStyleService> GetStyleService();
+
 private:
 
 	static TSharedRef< class FSlateStyleSet > Create(FFriendsAndChatStyle FriendStyle);
@@ -114,4 +116,5 @@ private:
 private:
 
 	static TSharedPtr< class FSlateStyleSet > FriendsAndChatModuleStyleInstance;
+	static TSharedPtr< class FFriendsFontStyleService > FriendsFontStyleService;
 };
diff --git a/Engine/Source/Developer/FriendsAndChat/Public/FriendsFontStyleService.h b/Engine/Source/Developer/FriendsAndChat/Public/FriendsFontStyleService.h
new file mode 100644
--- /dev/null
+++ b/Engine/Source/Developer/FriendsAndChat/Public/FriendsFontStyleService.h
@@ -0,0 +1,25 @@
+// Copyright 1998-2015 Epic Games, Inc. All Rights Reserved.
+
+#pragma once
+
+class FFriendsFontStyleService
+	: public TSharedFromThis<FFriendsFontStyleService>
+{
+public:
+	virtual ~FFriendsFontStyleService() {}
+
+
+	virtual void SetFontStyles(const FFriendsFontStyle& FriendsNormalFontStyle) = 0;
+
+	virtual void SetUserFontSize(EChatFontType::Type FontSize) = 0;
+
+	virtual FSlateFontInfo GetNormalFont() = 0;
+	virtual FSlateFontInfo GetNormalBoldFont() = 0;
+};
+
+/**
+* Creates the implementation for an FFriendsFontStyleService.
+*
+* @return the newly created FFriendsFontStyleService implementation.
+*/
+FACTORY(TSharedRef< FFriendsFontStyleService >, FFriendsFontStyleService);
\ No newline at end of file
diff --git a/Engine/Source/Developer/FriendsAndChat/Public/Layers/Presentation/IChatSettingsService.h b/Engine/Source/Developer/FriendsAndChat/Public/Layers/Presentation/IChatSettingsService.h
--- a/Engine/Source/Developer/FriendsAndChat/Public/Layers/Presentation/IChatSettingsService.h
+++ b/Engine/Source/Developer/FriendsAndChat/Public/Layers/Presentation/IChatSettingsService.h
@@ -42,6 +42,12 @@ public:
 	 */
 	virtual bool GetRadioOption(EChatSettingsType::Type Option, uint8 Index) const = 0;
 
+	/*
+	* Get an option that is stored as an enumeration value
+	* @param Option The Option to get the value for
+	*/
+	virtual uint8 GetRadioOptionIndex(EChatSettingsType::Type Option) const = 0;
+
 	/*
 	 * Set an option that is stored as an enumeration value
 	 * @param Option The Option to set the value for
