[INTEGRATE] merge UE4-LauncherDev > UE4 CL 2608884
Join game is greyed out when trying to join a full party

diff --git a/Engine/Source/Developer/FriendsAndChat/Private/Layers/Domain/ClanRepository.cpp b/Engine/Source/Developer/FriendsAndChat/Private/Layers/Domain/ClanRepository.cpp
--- a/Engine/Source/Developer/FriendsAndChat/Private/Layers/Domain/ClanRepository.cpp
+++ b/Engine/Source/Developer/FriendsAndChat/Private/Layers/Domain/ClanRepository.cpp
@@ -123,6 +123,16 @@ public:
 		return false;
 	}
 
+	virtual bool IsInParty() const override
+	{
+		return false;
+	}
+
+	virtual bool CanJoinParty() const override
+	{
+		return false;
+	}
+
 	virtual bool CanInvite() const override
 	{
 		return false;
diff --git a/Engine/Source/Developer/FriendsAndChat/Private/Layers/Presentation/FriendViewModel.cpp b/Engine/Source/Developer/FriendsAndChat/Private/Layers/Presentation/FriendViewModel.cpp
--- a/Engine/Source/Developer/FriendsAndChat/Private/Layers/Presentation/FriendViewModel.cpp
+++ b/Engine/Source/Developer/FriendsAndChat/Private/Layers/Presentation/FriendViewModel.cpp
@@ -46,12 +46,16 @@ public:
 			{
 				case EInviteStatus::Accepted :
 				{
-					if (FriendItem->IsOnline() && FriendItem->IsGameJoinable())
+					if (FriendItem->IsOnline() && FriendItem->IsGameJoinable() || FriendItem->IsInParty())
 					{
 						if(CanPerformAction(EFriendActionType::JoinGame))
 						{
 							Actions.Add(EFriendActionType::JoinGame);
 						}
+						else if(FriendsAndChatManager.Pin()->JoinGameAllowed(GetClientId()) && FriendItem->IsInParty())
+						{
+							Actions.Add(EFriendActionType::JoinGame);
+						}
 					}
 					if (FriendItem->IsOnline() && FriendItem->CanInvite())
 					{
@@ -162,7 +166,15 @@ public:
 		{
 			case EFriendActionType::JoinGame:
 			{
-				return FriendsAndChatManager.Pin()->JoinGameAllowed(GetClientId());
+				if(FriendsAndChatManager.Pin()->JoinGameAllowed(GetClientId()))
+				{
+					if(FriendItem->IsInParty())
+					{
+						return FriendItem->CanJoinParty();
+					}
+					return true;
+				}
+				return false;
 			}
 			case EFriendActionType::AcceptFriendRequest:
 			case EFriendActionType::RemoveFriend:
diff --git a/Engine/Source/Developer/FriendsAndChat/Private/Models/FriendItem.cpp b/Engine/Source/Developer/FriendsAndChat/Private/Models/FriendItem.cpp
--- a/Engine/Source/Developer/FriendsAndChat/Private/Models/FriendItem.cpp
+++ b/Engine/Source/Developer/FriendsAndChat/Private/Models/FriendItem.cpp
@@ -142,17 +142,23 @@ bool FFriendItem::IsGameJoinable() const
 			{
 				return true;
 			}
-			//@todo samz1 - determine if party is joinable
-			const bool bIsPartyJoinable = GetPartyJoinInfo().IsValid() && !FriendsAndChatManager.Pin()->IsFriendInSameParty(AsShared());
-			if (bIsPartyJoinable)
-			{
-				return true;
-			}
 		}
 	}
 	return false;
 }
 
+bool FFriendItem::IsInParty() const
+{
+	TSharedPtr<IOnlinePartyJoinInfo> PartyInfo = GetPartyJoinInfo();
+	return PartyInfo.IsValid();
+}
+
+bool FFriendItem::CanJoinParty() const
+{
+	TSharedPtr<IOnlinePartyJoinInfo> PartyInfo = GetPartyJoinInfo();
+	return PartyInfo.IsValid() && PartyInfo->GetIsAcceptingMembers() && !FriendsAndChatManager.Pin()->IsFriendInSameParty(AsShared());
+}
+
 bool FFriendItem::CanInvite() const
 {
 	FString FriendsClientID = GetClientId();
diff --git a/Engine/Source/Developer/FriendsAndChat/Private/Models/FriendItem.h b/Engine/Source/Developer/FriendsAndChat/Private/Models/FriendItem.h
--- a/Engine/Source/Developer/FriendsAndChat/Private/Models/FriendItem.h
+++ b/Engine/Source/Developer/FriendsAndChat/Private/Models/FriendItem.h
@@ -103,6 +103,18 @@ public:
 	 */
 	virtual bool IsGameJoinable() const override;
 
+	/**
+	 * Get if the user is in a party
+	 * @return The user in party state.
+	 */
+	virtual bool IsInParty() const override;
+
+	/**
+	 * Get if the user is in a joinable party
+	 * @return The user joinable party state.
+	 */
+	virtual bool CanJoinParty() const override;
+
 	/**
 	* Get if the user can join our game if we were to invite them
 	* @return True if we can invite them
diff --git a/Engine/Source/Developer/FriendsAndChat/Private/Models/FriendRecentPlayerItem.cpp b/Engine/Source/Developer/FriendsAndChat/Private/Models/FriendRecentPlayerItem.cpp
--- a/Engine/Source/Developer/FriendsAndChat/Private/Models/FriendRecentPlayerItem.cpp
+++ b/Engine/Source/Developer/FriendsAndChat/Private/Models/FriendRecentPlayerItem.cpp
@@ -110,6 +110,16 @@ bool FFriendRecentPlayerItem::CanInvite() const
 	return false;
 }
 
+bool FFriendRecentPlayerItem::IsInParty() const
+{
+	return false;
+}
+
+bool FFriendRecentPlayerItem::CanJoinParty() const
+{
+	return false;
+}
+
 TSharedPtr<const FUniqueNetId> FFriendRecentPlayerItem::GetGameSessionId() const
 {
 	return nullptr;
diff --git a/Engine/Source/Developer/FriendsAndChat/Private/Models/FriendRecentPlayerItems.h b/Engine/Source/Developer/FriendsAndChat/Private/Models/FriendRecentPlayerItems.h
--- a/Engine/Source/Developer/FriendsAndChat/Private/Models/FriendRecentPlayerItems.h
+++ b/Engine/Source/Developer/FriendsAndChat/Private/Models/FriendRecentPlayerItems.h
@@ -127,6 +127,18 @@ public:
 	/** Is the player in a game that is joinable */
 	virtual bool IsGameJoinable() const override;
 
+	/**
+	 * Get if the user is in a party
+	 * @return The user in party state.
+	 */
+	virtual bool IsInParty() const override;
+
+	/**
+	 * Get if the user is in a joinable party
+	 * @return The user joinable party state.
+	 */
+	virtual bool CanJoinParty() const override;
+
 	/** Get if the user can join our game if we were to invite them*/
 	virtual bool CanInvite() const override;
 
diff --git a/Engine/Source/Developer/FriendsAndChat/Public/Models/IFriendItem.h b/Engine/Source/Developer/FriendsAndChat/Public/Models/IFriendItem.h
--- a/Engine/Source/Developer/FriendsAndChat/Public/Models/IFriendItem.h
+++ b/Engine/Source/Developer/FriendsAndChat/Public/Models/IFriendItem.h
@@ -120,6 +120,18 @@ public:
 	 */
 	virtual bool IsGameJoinable() const = 0;
 
+	/**
+	 * Get if the user is in a party
+	 * @return The user in party state.
+	 */
+	virtual bool IsInParty() const = 0;
+
+	/**
+	 * Get if the user is in a joinable party
+	 * @return The user joinable party state.
+	 */
+	virtual bool CanJoinParty() const = 0;
+
 	/**
 	 * Get if the user can join our game if we were to invite them
 	 * @return True if we can invite them
