[INTEGRATE] merge UE4-LauncherDev > UE4 CL 2608969

diff --git a/Engine/Source/Developer/FriendsAndChat/Private/Layers/Presentation/FriendViewModel.cpp b/Engine/Source/Developer/FriendsAndChat/Private/Layers/Presentation/FriendViewModel.cpp
--- a/Engine/Source/Developer/FriendsAndChat/Private/Layers/Presentation/FriendViewModel.cpp
+++ b/Engine/Source/Developer/FriendsAndChat/Private/Layers/Presentation/FriendViewModel.cpp
@@ -283,6 +283,10 @@ private:
 		{
 			FriendsAndChatManager.Pin()->AcceptGameInvite(FriendItem);
 		}
+		else if(FriendsAndChatManager.Pin()->JoinGameAllowed(GetClientId()) && FriendItem->CanJoinParty())
+		{
+			FriendsAndChatManager.Pin()->AcceptGameInvite(FriendItem);
+		}
 	}
 
 	void RejectGame()
