[INTEGRATE] merge UE4-LauncherDev > UE4 FriendsAndChat up to CL 2609064
Fix Linux build error for CL 2608984

diff --git a/Engine/Source/Developer/FriendsAndChat/FriendsAndChat.Build.cs b/Engine/Source/Developer/FriendsAndChat/FriendsAndChat.Build.cs
--- a/Engine/Source/Developer/FriendsAndChat/FriendsAndChat.Build.cs
+++ b/Engine/Source/Developer/FriendsAndChat/FriendsAndChat.Build.cs
@@ -34,6 +34,7 @@ public class FriendsAndChat : ModuleRules
 				"Developer/FriendsAndChat/Private/Models",
 				"Developer/FriendsAndChat/Private/Core",
 				"Developer/FriendsAndChat/Private/Layers/DataAccess",
+		 		"Developer/FriendsAndChat/Private/Layers/DataAccess/Analytics",
 				"Developer/FriendsAndChat/Private/Layers/Domain",
 				"Developer/FriendsAndChat/Private/Layers/Presentation",
 				"Developer/FriendsAndChat/Private/Layers/UI"
diff --git a/Engine/Source/Developer/FriendsAndChat/Private/FriendsAndChatStyle.cpp b/Engine/Source/Developer/FriendsAndChat/Private/FriendsAndChatStyle.cpp
--- a/Engine/Source/Developer/FriendsAndChat/Private/FriendsAndChatStyle.cpp
+++ b/Engine/Source/Developer/FriendsAndChat/Private/FriendsAndChatStyle.cpp
@@ -35,6 +35,12 @@ FFriendsAndChatStyle& FFriendsAndChatStyle::SetFriendsChatChromeStyle(const FFri
 	return *this;
 }
 
+FFriendsAndChatStyle& FFriendsAndChatStyle::SetFriendsMarkupStyle(const FFriendsMarkupStyle& InFriendsMarkupStyle)
+{
+	FriendsMarkupStyle = InFriendsMarkupStyle;
+	return *this;
+}
+
 FFriendsAndChatStyle& FFriendsAndChatStyle::SetWindowEdgingBrush(const FSlateBrush& Value)
 {
 	WindowEdgingBrush = Value;
diff --git a/Engine/Source/Developer/FriendsAndChat/Private/FriendsMarkupStyle.cpp b/Engine/Source/Developer/FriendsAndChat/Private/FriendsMarkupStyle.cpp
new file mode 100644
--- /dev/null
+++ b/Engine/Source/Developer/FriendsAndChat/Private/FriendsMarkupStyle.cpp
@@ -0,0 +1,23 @@
+// Copyright 1998-2015 Epic Games, Inc. All Rights Reserved.
+
+#include "FriendsAndChatPrivatePCH.h"
+
+const FName FFriendsMarkupStyle::TypeName( TEXT("FFriendsMarkupStyle") );
+
+const FFriendsMarkupStyle& FFriendsMarkupStyle::GetDefault()
+{
+	static FFriendsMarkupStyle Default;
+	return Default;
+}
+
+FFriendsMarkupStyle& FFriendsMarkupStyle::SetMarkupButtonStyle(const FButtonStyle& ButtonStyle)
+{
+	MarkupButtonStyle = ButtonStyle;
+	return *this;
+}
+
+FFriendsMarkupStyle& FFriendsMarkupStyle::SetMarkupTextStyle(const FTextBlockStyle& InTextStyle)
+{
+	MarkupTextStyle = InTextStyle;
+	return *this;
+}
diff --git a/Engine/Source/Developer/FriendsAndChat/Private/Layers/DataAccess/Analytics/FriendsAndChatAnalytics.h b/Engine/Source/Developer/FriendsAndChat/Private/Layers/DataAccess/Analytics/FriendsAndChatAnalytics.h
new file mode 100644
--- /dev/null
+++ b/Engine/Source/Developer/FriendsAndChat/Private/Layers/DataAccess/Analytics/FriendsAndChatAnalytics.h
@@ -0,0 +1,168 @@
+// Copyright 1998-2015 Epic Games, Inc. All Rights Reserved.
+#pragma once
+
+class IAnalyticsProvider;
+
+class FFriendsAndChatAnalytics
+{
+public:
+
+	virtual ~FFriendsAndChatAnalytics() { }
+
+	/**
+	 * Update provider to use for capturing events
+	 */
+	virtual void SetProvider(const TSharedPtr<IAnalyticsProvider>& InAnalyticsProvider) = 0;
+
+	/*
+	 * @EventName Social.GameInvite.Accept
+	 *
+	 * @Trigger Fires whenever a game invite is accepted
+	 *
+	 * @Type Static
+	 *
+	 * @EventParam User (string) The local user ID
+	 * @EventParam Friend (string) The inviter ID
+	 * @EventParam ClientId (string) the client ID
+	 * @EventParam Status (string) presence status
+	 *
+	 * @Comments 
+	 *
+	 */
+	virtual void FireEvent_AcceptGameInvite(const FUniqueNetId& LocalUserId, const FUniqueNetId& ToUser) const = 0;
+
+	/*
+	 * @EventName Social.GameInvite.Reject
+	 *
+	 * @Trigger Fires whenever a game invite is rejected
+	 *
+	 * @Type Static
+	 *
+	 * @EventParam User (string) The local user ID
+	 * @EventParam Friend (string) The inviter ID
+	 * @EventParam ClientId (string) the client ID
+	 * @EventParam Status (string) presence status
+	 *
+	 * @Comments 
+	 *
+	 */
+	virtual void FireEvent_RejectGameInvite(const FUniqueNetId& LocalUserId, const FUniqueNetId& ToUser) const = 0;
+
+	/*
+	 * @EventName Social.GameInvite.Send
+	 *
+	 * @Trigger Fires whenever a game invite is sent
+	 *
+	 * @Type Static
+	 *
+	 * @EventParam User (string) The local user ID
+	 * @EventParam Friend (string) The invitee ID
+	 * @EventParam ClientId (string) the client ID
+	 * @EventParam Status (string) presence status
+	 *
+	 * @Comments 
+	 *
+	 */
+	virtual void FireEvent_SendGameInvite(const FUniqueNetId& LocalUserId, const FUniqueNetId& ToUser) const = 0;
+
+	/*
+	 * @EventName Social.AddFriend
+	 *
+	 * @Trigger Fires whenever we add a friend
+	 *
+	 * @Type Static
+	 *
+	 * @EventParam User (string) The local user ID
+	 * @EventParam Friend (string) The invitee ID
+	 * @EventParam FriendName(string) The friend name
+	 * @EventParam Result (string) The find friend result
+	 * @EventParam bRecentPlayer (bool) If this was from a recent player list
+	 * @EventParam ClientId (string) the client ID
+	 * @EventParam Status (string) presence status
+	 *
+	 * @Comments 
+	 *
+	 */
+	virtual void FireEvent_AddFriend(const FUniqueNetId& LocalUserId, const FString& FriendName, const FUniqueNetId& FriendId, EFindFriendResult::Type Result, bool bRecentPlayer) const = 0;
+
+
+	/*
+	 * @EventName Social.FriendAction.Accept
+	 *
+	 * @Trigger Fires whenever a user accepts a friend request
+	 *
+	 * @Type Static
+	 *
+	 * @EventParam User (string) The local user ID
+	 * @EventParam Friend (string) The invitee ID
+	 * @EventParam ClientId (string) the client ID
+	 * @EventParam Status (string) presence status
+	 *
+	 * @Comments 
+	 *
+	 */
+	virtual void FireEvent_FriendRequestAccepted(const FUniqueNetId& LocalUserId, const IFriendItem& Friend) const = 0;
+
+	/*
+	 * @EventName Social.FriendRemoved
+	 *
+	 * @Trigger Fires whenever a friend is removed
+	 *
+	 * @Type Static
+	 *
+	 * @EventParam User (string) The local user ID
+	 * @EventParam Friend (string) The invitee ID
+	 * @EventParam RemoveReason (string) The reason a friend was removed
+	 * @EventParam ClientId (string) the client ID
+	 * @EventParam Status (string) presence status
+	 *
+	 * @Comments 
+	 *
+	 */
+	virtual void FireEvent_FriendRemoved(const FUniqueNetId& LocalUserId, const IFriendItem& Friend, const FString& RemoveReason) const = 0;
+
+	/*
+	 * @EventName Social.Chat.Toggle
+	 *
+	 * @Trigger Fires whenever toggle global chat on or off
+	 *
+	 * @Type Static
+	 *
+	 * @EventParam User (string) The local user ID
+	 * @EventParam Channel (string) Channel Opened
+	 * @EventParam bEnabled (bool) if enabled
+	 * @EventParam ClientId (string) the client ID
+	 * @EventParam Status (string) presence status
+	 *
+	 * @Comments 
+	 *
+	 */
+	virtual void FireEvent_RecordToggleChat(const FUniqueNetId& LocalUserId, const FString& Channel, bool bEnabled) const = 0;
+
+	/**
+	 * Record a private chat to a user (aggregates pending FlushChat)
+	 */
+	virtual void RecordPrivateChat(const FString& ToUser) = 0;
+	/**
+	 * Record a public chat to a channel (aggregates pending FlushChat)
+	 */
+	virtual void RecordChannelChat(const FString& ToChannel) = 0;
+
+	/*
+	 * @EventName Social.Chat.Counts.2
+	 *
+	 * @Trigger Fires when the user logs out, and periodically on CHAT_ANALYTICS_INTERVAL. Currently every 5 minutes
+	 *
+	 * @Type Static
+	 *
+	 * @EventParam Name(string) the channel or user
+	 * @EventParam Type (string) The chat type
+	 * @EventParam Count (int) amount of times the event occurred
+
+	 * @Comments 
+	 *
+	 */
+	virtual void FireEvent_FlushChatStats() = 0;
+};
+
+FACTORY(TSharedRef<class FFriendsAndChatAnalytics>, FFriendsAndChatAnalytics);
diff --git a/Engine/Source/Developer/FriendsAndChat/Private/Layers/DataAccess/FriendsAndChatManager.cpp b/Engine/Source/Developer/FriendsAndChat/Private/Layers/DataAccess/FriendsAndChatManager.cpp
--- a/Engine/Source/Developer/FriendsAndChat/Private/Layers/DataAccess/FriendsAndChatManager.cpp
+++ b/Engine/Source/Developer/FriendsAndChat/Private/Layers/DataAccess/FriendsAndChatManager.cpp
@@ -26,6 +26,7 @@
 #include "FriendsChatMarkupService.h"
 #include "ChatDisplayService.h"
 #include "ChatSettingsService.h"
+#include "FriendsAndChatAnalytics.h"
 
 #define LOCTEXT_NAMESPACE "FriendsAndChatManager"
 
@@ -57,6 +58,7 @@ void FFriendsAndChatManager::Initialize()
 	FriendsListFactory = FFriendListFactoryFactory::Create(FriendViewModelFactory.ToSharedRef(), SharedThis(this));
 	TSharedRef<IChatCommunicationService> CommunicationService = StaticCastSharedRef<IChatCommunicationService>(MessageManager.ToSharedRef());
 	MarkupServiceFactory = FFriendsChatMarkupServiceFactoryFactory::Create(CommunicationService, NavigationService.ToSharedRef(), FriendsListFactory.ToSharedRef());
+	Analytics = FFriendsAndChatAnalyticsFactory::Create();
 }
 
 /* IFriendsAndChatManager interface
@@ -114,7 +116,6 @@ void FFriendsAndChatManager::Login(IOnlineSubsystem* InOnlineSub, bool bInIsGame
 			OnDestroySessionCompleteDelegate     = FOnDestroySessionCompleteDelegate               ::CreateSP (this, &FFriendsAndChatManager::OnGameDestroyed);
 			OnPartyMemberJoinedDelegate          = FOnPartyMemberJoinedDelegate                    ::CreateSP (this, &FFriendsAndChatManager::OnPartyMemberJoined);
 			OnPartyMemberExitedDelegate          = FOnPartyMemberExitedDelegate                    ::CreateSP (this, &FFriendsAndChatManager::OnPartyMemberExited);
-			OnPartyJoinedDelegate = FOnPartyJoinedDelegate::CreateSP(this, &FFriendsAndChatManager::OnPartyJoined);
 
 			OnQueryRecentPlayersCompleteDelegateHandle	= FriendsInterface->AddOnQueryRecentPlayersCompleteDelegate_Handle(   OnQueryRecentPlayersCompleteDelegate);
 			OnFriendsListChangedDelegateHandle			= FriendsInterface->AddOnFriendsChangeDelegate_Handle(LocalControllerIndex, OnFriendsListChangedDelegate);
@@ -128,12 +129,11 @@ void FFriendsAndChatManager::Login(IOnlineSubsystem* InOnlineSub, bool bInIsGame
 
 			OnPresenceReceivedCompleteDelegateHandle = OnlineSub->GetPresenceInterface()->AddOnPresenceReceivedDelegate_Handle      (OnPresenceReceivedCompleteDelegate);
 			OnGameInviteReceivedDelegateHandle       = OnlineSub->GetSessionInterface ()->AddOnSessionInviteReceivedDelegate_Handle (OnGameInviteReceivedDelegate);
-			OnPartyInviteReceivedDelegateHandle      = OnlineSub->GetPartyInterface()->AddOnPartyInviteReceivedDelegate_Handle      (OnPartyInviteReceivedDelegate);
+			OnPartyInviteReceivedDelegateHandle = OnlineSub->GetPartyInterface()->AddOnPartyInviteReceivedDelegate_Handle(OnPartyInviteReceivedDelegate);
 			OnDestroySessionCompleteDelegateHandle   = OnlineSub->GetSessionInterface ()->AddOnDestroySessionCompleteDelegate_Handle(OnDestroySessionCompleteDelegate);
 
 			OnPartyMemberJoinedDelegateHandle = OnlineSub->GetPartyInterface()->AddOnPartyMemberJoinedDelegate_Handle(OnPartyMemberJoinedDelegate);
 			OnPartyMemberExitedDelegateHandle = OnlineSub->GetPartyInterface()->AddOnPartyMemberExitedDelegate_Handle(OnPartyMemberExitedDelegate);
-			OnPartyJoinedDelegateHandle       = OnlineSub->GetPartyInterface()->AddOnPartyJoinedDelegate_Handle      (OnPartyJoinedDelegate);
 
 			ManagerState = EFriendsAndManagerState::Idle;
 
@@ -170,7 +170,7 @@ void FFriendsAndChatManager::Login(IOnlineSubsystem* InOnlineSub, bool bInIsGame
 void FFriendsAndChatManager::Logout()
 {
 	// flush before removing the analytics provider
-	Analytics.FlushChatStats(); 
+	Analytics->FireEvent_FlushChatStats(); 
 
 	if (OnlineSub != nullptr)
 	{
@@ -205,7 +205,6 @@ void FFriendsAndChatManager::Logout()
 		{
 			OnlineSub->GetPartyInterface()->ClearOnPartyMemberJoinedDelegate_Handle(OnPartyMemberJoinedDelegateHandle);
 			OnlineSub->GetPartyInterface()->ClearOnPartyMemberExitedDelegate_Handle(OnPartyMemberExitedDelegateHandle);
-			OnlineSub->GetPartyInterface()->ClearOnPartyJoinedDelegate_Handle(OnPartyJoinedDelegateHandle);
 		}
 	}
 
@@ -270,7 +269,7 @@ void FFriendsAndChatManager::ClearApplicationViewModels()
 
 void FFriendsAndChatManager::SetAnalyticsProvider(const TSharedPtr<IAnalyticsProvider>& AnalyticsProvider)
 {
-	Analytics.SetProvider(AnalyticsProvider);
+	Analytics->SetProvider(AnalyticsProvider);
 }
 
 void FFriendsAndChatManager::InsertNetworkChatMessage(const FString& InMessage)
@@ -393,6 +392,7 @@ TSharedPtr< SWidget > FFriendsAndChatManager::GenerateStatusWidget( const FFrien
 TSharedPtr< SWidget > FFriendsAndChatManager::GenerateChatWidget(const FFriendsAndChatStyle* InStyle, TAttribute<FText> ActivationHintDelegate, EChatMessageType::Type MessageType, TSharedPtr<IFriendItem> FriendItem, TSharedPtr< SWindow > WindowPtr)
 {
 	TSharedRef<IChatDisplayService> ChatDisplayService = GenerateChatDisplayService();
+	Style = *InStyle;
 
 	TSharedPtr<FChatViewModel> ViewModel = FChatViewModelFactory::Create(MessageManager.ToSharedRef(), NavigationService.ToSharedRef(), MarkupServiceFactory->Create(), ChatDisplayService, SharedThis(this), EChatViewModelType::Windowed);
 	ViewModel->SetChannelFlags(MessageType);
@@ -487,7 +487,7 @@ void FFriendsAndChatManager::OpenGlobalChat()
 		if (OnlineIdentity.IsValid() &&
 			OnlineIdentity->GetUniquePlayerId(LocalControllerIndex).IsValid())
 		{
-			Analytics.RecordToggleChat(*OnlineIdentity->GetUniquePlayerId(LocalControllerIndex), TEXT("Global"), true, TEXT("Social.Chat.Toggle"));
+			Analytics->FireEvent_RecordToggleChat(*OnlineIdentity->GetUniquePlayerId(LocalControllerIndex), TEXT("Global"), true);
 			NavigationService->ChangeViewChannel(EChatMessageType::Global);
 		}
 	}
@@ -519,12 +519,11 @@ void FFriendsAndChatManager::AcceptFriend( TSharedPtr< IFriendItem > FriendItem
 {
 	PendingOutgoingAcceptFriendRequests.Add(FriendItem->GetOnlineUser()->GetUserId()->ToString());
 	FriendItem->SetPendingAccept();
-	RefreshList();
 	OnFriendsNotification().Broadcast(false);
 	if (OnlineIdentity.IsValid() &&
 		OnlineIdentity->GetUniquePlayerId(LocalControllerIndex).IsValid())
 	{
-		Analytics.RecordFriendAction(*OnlineIdentity->GetUniquePlayerId(LocalControllerIndex), *FriendItem, TEXT("Social.FriendAction.Accept"));
+		Analytics->FireEvent_FriendRequestAccepted(*OnlineIdentity->GetUniquePlayerId(LocalControllerIndex), *FriendItem);
 	}
 }
 
@@ -534,17 +533,16 @@ void FFriendsAndChatManager::RejectFriend(TSharedPtr< IFriendItem > FriendItem)
 	NotifiedRequest.Remove( FriendItem->GetOnlineUser()->GetUserId() );
 	PendingOutgoingDeleteFriendRequests.Add(FriendItem->GetOnlineUser()->GetUserId().Get().ToString());
 	FriendsList.Remove( FriendItem );
-	RefreshList();
 	OnFriendsNotification().Broadcast(false);
 
 	if (OnlineIdentity.IsValid() &&
 		OnlineIdentity->GetUniquePlayerId(LocalControllerIndex).IsValid())
 	{
-		Analytics.RecordFriendAction(*OnlineIdentity->GetUniquePlayerId(LocalControllerIndex), *FriendItem, TEXT("Social.FriendAction.Reject"));
+		Analytics->FireEvent_FriendRemoved(*OnlineIdentity->GetUniquePlayerId(LocalControllerIndex), *FriendItem, EFriendActionType::ToAnalyticString(EFriendActionType::RejectFriendRequest));
 	}
 }
 
-void FFriendsAndChatManager::DeleteFriend(TSharedPtr< IFriendItem > FriendItem, const FString& Action)
+void FFriendsAndChatManager::DeleteFriend(TSharedPtr< IFriendItem > FriendItem, EFriendActionType::Type Reason)
 {
 	TSharedRef<const FUniqueNetId> UserID = FriendItem->GetOnlineFriend()->GetUserId();
 	NotifiedRequest.Remove( UserID );
@@ -558,7 +556,7 @@ void FFriendsAndChatManager::DeleteFriend(TSharedPtr< IFriendItem > FriendItem,
 	if (OnlineIdentity.IsValid() &&
 		OnlineIdentity->GetUniquePlayerId(LocalControllerIndex).IsValid())
 	{
-		Analytics.RecordFriendAction(*OnlineIdentity->GetUniquePlayerId(LocalControllerIndex), *FriendItem, TEXT("Social.FriendAction.") + Action);
+		Analytics->FireEvent_FriendRemoved(*OnlineIdentity->GetUniquePlayerId(LocalControllerIndex), *FriendItem, EFriendActionType::ToAnalyticString(Reason));
 	}
 }
 
@@ -602,7 +600,7 @@ FReply FFriendsAndChatManager::HandleMessageAccepted( TSharedPtr< FFriendsAndCha
 			TSharedPtr< IFriendItem > User = FindUser(ChatMessage->GetUniqueID().Get());
 			if (User.IsValid())
 			{
-				DeleteFriend(User, EFriendActionType::ToText(EFriendActionType::IgnoreFriendRequest).ToString());
+				DeleteFriend(User, EFriendActionType::IgnoreFriendRequest);
 			}
 		}
 	break;
@@ -613,6 +611,25 @@ FReply FFriendsAndChatManager::HandleMessageAccepted( TSharedPtr< FFriendsAndCha
 	return FReply::Handled();
 }
 
+FReply FFriendsAndChatManager::HandleGameInviteResponse( TSharedPtr< FFriendsAndChatMessage > MessageNotification, EFriendsResponseType::Type ResponseType )
+{
+	TSharedPtr< IFriendItem > User = FindUser(MessageNotification->GetUniqueID().Get());
+	if(User.IsValid())
+	{
+		if(ResponseType == EFriendsResponseType::Response_Accept)
+		{
+			AcceptGameInvite(User);
+		}
+		else
+		{
+			RejectGameInvite(User);
+		}
+	}
+
+	NotificationMessages.Remove(MessageNotification);
+	return FReply::Handled();
+}
+
 // Getters
 
 int32 FFriendsAndChatManager::GetFilteredFriendsList( TArray< TSharedPtr< IFriendItem > >& OutFriendsList )
@@ -960,7 +977,7 @@ bool FFriendsAndChatManager::Tick( float Delta )
 	FlushChatAnalyticsCountdown -= Delta;
 	if (FlushChatAnalyticsCountdown <= 0)
 	{
-		Analytics.FlushChatStats();
+		Analytics->FireEvent_FlushChatStats();
 		// Reset countdown for new update
 		FlushChatAnalyticsCountdown = CHAT_ANALYTICS_INTERVAL;
 	}
@@ -1207,6 +1224,7 @@ void FFriendsAndChatManager::OnQueryUserInfoComplete( int32 LocalPlayer, bool bW
 				{
 					TSharedPtr< FFriendItem > FriendItem = MakeShareable(new FFriendItem(OnlineFriend, OnlineUser, EFriendsDisplayLists::DefaultDisplay, SharedThis(this)));
 					PendingFriendsList.Add(FriendItem);
+					FriendItem->SetChachedOnline(FriendItem->IsOnline());
 				}
 			}
 			else
@@ -1458,7 +1476,7 @@ void FFriendsAndChatManager::SendMessageReceivedNotification(const TSharedPtr< I
 {
 	if (OnFriendsActionNotification().IsBound())
 	{
-		const FText FriendReceivedMessage = FText::Format(LOCTEXT("FriendAddedToast", "Message from {0}"), FText::FromString(Friend->GetName()));
+		const FText FriendReceivedMessage = FText::Format(LOCTEXT("FriendMessageRecievedToast", "Message from {0}"), FText::FromString(Friend->GetName()));
 		TSharedPtr< FFriendsAndChatMessage > NotificationMessage = MakeShareable(new FFriendsAndChatMessage(FriendReceivedMessage.ToString()));
 		NotificationMessage->SetMessageType(EFriendsRequestType::ChatMessage);
 		OnFriendsActionNotification().Broadcast(NotificationMessage.ToSharedRef());
@@ -1531,12 +1549,11 @@ void FFriendsAndChatManager::OnQueryUserIdMappingComplete(bool bWasSuccessful, c
 	if (OnlineIdentity.IsValid() &&
 		OnlineIdentity->GetUniquePlayerId(LocalControllerIndex).IsValid())
 	{
-		Analytics.RecordAddFriend(*OnlineIdentity->GetUniquePlayerId(LocalControllerIndex), 
+		Analytics->FireEvent_AddFriend(*OnlineIdentity->GetUniquePlayerId(LocalControllerIndex), 
 								  DisplayName, 
 								  IdentifiedUserId, 
 								  FindFriendResult, 
-								  bRecentPlayer, 
-								  TEXT("Social.AddFriend"));
+								  bRecentPlayer);
 	}
 
 	FriendByNameRequests.Remove( DisplayName );
@@ -1553,7 +1570,6 @@ void FFriendsAndChatManager::OnQueryUserIdMappingComplete(bool bWasSuccessful, c
 		}
 		else
 		{
-			RefreshList();
 			SetState(EFriendsAndManagerState::Idle);
 		}
 	}
@@ -1566,7 +1582,6 @@ void FFriendsAndChatManager::OnSendInviteComplete( int32 LocalPlayer, bool bWasS
 
 	if ( PendingOutgoingFriendRequests.Num() == 0 )
 	{
-		RefreshList();
 		RequestListRefresh();
 		SetState(EFriendsAndManagerState::Idle);
 	}
@@ -1574,13 +1589,9 @@ void FFriendsAndChatManager::OnSendInviteComplete( int32 LocalPlayer, bool bWasS
 
 void FFriendsAndChatManager::OnPresenceReceived(const FUniqueNetId& UserId, const TSharedRef<FOnlineUserPresence>& NewPresence)
 {
-	RefreshList();
-
 	// Compare to previous presence for this friend, display a toast if a friend has came online or joined a game
 	TSharedPtr<const FUniqueNetId> SelfId = OnlineIdentity->GetUniquePlayerId(0);
 	bool bIsSelf = (SelfId.IsValid() && UserId == *SelfId);
-	TSharedPtr<IFriendItem> PresenceFriend = FindUser(UserId);
-	bool bFoundFriend = PresenceFriend.IsValid();
 	bool bFriendNotificationsBound = OnFriendsActionNotification().IsBound();
 	// Don't show notifications if we're building the friends list presences for the first time.
 	// Guess at this using the size of the saved presence list. OK to show the last 1 or 2, but avoid spamming dozens of notifications at startup
@@ -1628,6 +1639,16 @@ void FFriendsAndChatManager::OnPresenceReceived(const FUniqueNetId& UserId, cons
 		}
 	}
 
+	TSharedPtr<IFriendItem> PresenceFriend = FindUser(UserId);
+	if(PresenceFriend.IsValid())
+	{
+		if(PresenceFriend->IsCachedOnline() != PresenceFriend->IsOnline())
+		{
+			RefreshList();
+		}
+		PresenceFriend->SetChachedOnline(PresenceFriend->IsOnline());
+	}
+
 // @todo Antony.Carter Disabled for GDC as its a bit spammy
 #if 0
 	// When a new friend comes online, we should see eg. OnlineFriedCount = 3, OldUserPresence = 2.  Only assume we're starting up if there's a difference of 2 or more
@@ -1671,7 +1692,7 @@ void FFriendsAndChatManager::OnPresenceReceived(const FUniqueNetId& UserId, cons
 					&& NewPresence->SessionId.IsValid()
 					&& NewPresence->SessionId->IsValid())
 				{
-					const FText PresenceChangeText = FText::Format(LOCTEXT("FriendPresenceChange_Online", "{FriendName} Is Now In A Game"), Args);
+					const FText PresenceChangeText = FText::Format(LOCTEXT("FriendPresenceChange_InGame", "{FriendName} Is Now In A Game"), Args);
 					TSharedPtr< FFriendsAndChatMessage > NotificationMessage = MakeShareable(new FFriendsAndChatMessage(PresenceChangeText.ToString()));
 					NotificationMessage->SetMessageType(EFriendsRequestType::PresenceChange);
 					OnFriendsActionNotification().Broadcast(NotificationMessage.ToSharedRef());
@@ -1681,7 +1702,7 @@ void FFriendsAndChatManager::OnPresenceReceived(const FUniqueNetId& UserId, cons
 				{
 					// could limit notifications to same game only by removing isPlaying check above
 					Args.Add(TEXT("GameName"), FText::FromString(PresenceFriend->GetClientName()));
-					const FText PresenceChangeText = FText::Format(LOCTEXT("FriendPresenceChange_Online", "{FriendName} Is Now Playing {GameName}"), Args);
+					const FText PresenceChangeText = FText::Format(LOCTEXT("FriendPresenceChange_NowPlayingThisGame", "{FriendName} Is Now Playing {GameName}"), Args);
 					TSharedPtr< FFriendsAndChatMessage > NotificationMessage = MakeShareable(new FFriendsAndChatMessage(PresenceChangeText.ToString()));
 					NotificationMessage->SetMessageType(EFriendsRequestType::PresenceChange);
 					OnFriendsActionNotification().Broadcast(NotificationMessage.ToSharedRef());
@@ -1690,7 +1711,7 @@ void FFriendsAndChatManager::OnPresenceReceived(const FUniqueNetId& UserId, cons
 				else if (OldPresence.bIsPlaying == false && NewPresence->bIsPlaying == true)
 				{
 					Args.Add(TEXT("GameName"), FText::FromString(PresenceFriend->GetClientName()));
-					const FText PresenceChangeText = FText::Format(LOCTEXT("FriendPresenceChange_Online", "{FriendName} Is Now Playing {GameName}"), Args);
+					const FText PresenceChangeText = FText::Format(LOCTEXT("FriendPresenceChange_Nowplaying", "{FriendName} Is Now Playing {GameName}"), Args);
 					TSharedPtr< FFriendsAndChatMessage > NotificationMessage = MakeShareable(new FFriendsAndChatMessage(PresenceChangeText.ToString()));
 					NotificationMessage->SetMessageType(EFriendsRequestType::PresenceChange);
 					OnFriendsActionNotification().Broadcast(NotificationMessage.ToSharedRef());
@@ -1717,7 +1738,6 @@ void FFriendsAndChatManager::OnPresenceReceived(const FUniqueNetId& UserId, cons
 
 void FFriendsAndChatManager::OnPresenceUpdated(const FUniqueNetId& UserId, const bool bWasSuccessful)
 {
-	RefreshList();
 }
 
 void FFriendsAndChatManager::OnFriendsListChanged()
@@ -1909,9 +1929,16 @@ void FFriendsAndChatManager::SendGameInviteNotification(const TSharedPtr<IFriend
 		FFormatNamedArguments Args;
 		Args.Add(TEXT("Username"), FText::FromString(FriendItem->GetName()));
 		const FText FriendRequestMessage = FText::Format(LOCTEXT("FFriendsAndChatManager_GameInvite", "Game invite from {Username}"), Args);
-
-		TSharedPtr< FFriendsAndChatMessage > NotificationMessage = MakeShareable(new FFriendsAndChatMessage(FriendRequestMessage.ToString()));
+		TSharedPtr< FFriendsAndChatMessage > NotificationMessage = MakeShareable(new FFriendsAndChatMessage(FriendRequestMessage.ToString(), FriendItem->GetUniqueID()));
 		NotificationMessage->SetMessageType(EFriendsRequestType::GameInvite);
+
+		NotificationMessage->SetButtonCallback( FOnClicked::CreateSP(this, &FFriendsAndChatManager::HandleGameInviteResponse, NotificationMessage, EFriendsResponseType::Response_Accept));
+		NotificationMessage->SetButtonCallback( FOnClicked::CreateSP(this, &FFriendsAndChatManager::HandleGameInviteResponse, NotificationMessage, EFriendsResponseType::Response_Reject));
+		NotificationMessage->SetButtonDescription(LOCTEXT("FFriendsAndChatManager_GameInviteJoin", "Join Game"));
+		NotificationMessage->SetButtonDescription(LOCTEXT("FFriendsAndChatManager_GameInviteReject", "Reject"));
+		NotificationMessage->SetButtonStyle(TEXT("FriendsListEmphasisButton"));
+		NotificationMessage->SetButtonStyle(TEXT("FriendsListCriticalButton"));
+		NotificationMessage->SetMessageType(EFriendsRequestType::FriendInvite);
 		OnFriendsActionNotification().Broadcast(NotificationMessage.ToSharedRef());
 	}
 
@@ -1946,13 +1973,6 @@ void FFriendsAndChatManager::OnPartyMemberExited(const FUniqueNetId& LocalUserId
 	//GetChatViewModel()->UpdateInPartyUI();
 }
 
-void FFriendsAndChatManager::OnPartyJoined(const FUniqueNetId& LocalUserId, const FOnlinePartyId& PartyId)
-{
-	//ToDo - NickDavies - re-add party 
-	// You joined a party, update the UI to make party chat available & the default output channel
-	//GetChatViewModel()->UpdateInPartyUI();
-}
-
 void FFriendsAndChatManager::RejectGameInvite(const TSharedPtr<IFriendItem>& FriendItem)
 {
 	TSharedPtr<IFriendItem>* Existing = PendingGameInvitesList.Find(FriendItem->GetUniqueID()->ToString());
@@ -1978,9 +1998,9 @@ void FFriendsAndChatManager::RejectGameInvite(const TSharedPtr<IFriendItem>& Fri
 	if (OnlineIdentity.IsValid() &&
 		OnlineIdentity->GetUniquePlayerId(LocalControllerIndex).IsValid())
 	{
-		Analytics.RecordGameInvite(*OnlineIdentity->GetUniquePlayerId(LocalControllerIndex), *FriendItem->GetUniqueID(), TEXT("Social.GameInvite.Reject"));
+		Analytics->FireEvent_RejectGameInvite(*OnlineIdentity->GetUniquePlayerId(LocalControllerIndex), *FriendItem->GetUniqueID());
 	}
-	
+	OnFriendsNotification().Broadcast(false);
 }
 
 void FFriendsAndChatManager::AcceptGameInvite(const TSharedPtr<IFriendItem>& FriendItem)
@@ -2000,8 +2020,7 @@ void FFriendsAndChatManager::AcceptGameInvite(const TSharedPtr<IFriendItem>& Fri
 	if (PartyInfo.IsValid())
 	{
 		// notify for further processing of join party request 
-		const bool bIsFromInvite = FriendItem->GetListType() == EFriendsDisplayLists::GameInviteDisplay;
-		OnFriendsJoinParty().Broadcast(*FriendItem->GetUniqueID(), PartyInfo.ToSharedRef(), bIsFromInvite);
+		OnFriendsJoinParty().Broadcast(*FriendItem->GetUniqueID(), PartyInfo.ToSharedRef(), true);
 
 		AdditionalCommandline = OnlineSub->GetPartyInterface()->MakeTokenFromJoinInfo(*PartyInfo);
 	}
@@ -2026,8 +2045,9 @@ void FFriendsAndChatManager::AcceptGameInvite(const TSharedPtr<IFriendItem>& Fri
 	if (OnlineIdentity.IsValid() &&
 		OnlineIdentity->GetUniquePlayerId(LocalControllerIndex).IsValid())
 	{
-		Analytics.RecordGameInvite(*OnlineIdentity->GetUniquePlayerId(LocalControllerIndex), *FriendItem->GetUniqueID(), TEXT("Social.GameInvite.Accept"));
+		Analytics->FireEvent_AcceptGameInvite(*OnlineIdentity->GetUniquePlayerId(LocalControllerIndex), *FriendItem->GetUniqueID());
 	}
+	OnFriendsNotification().Broadcast(false);
 }
 
 void FFriendsAndChatManager::SendGameInvite(const TSharedPtr<IFriendItem>& FriendItem)
@@ -2079,7 +2099,7 @@ void FFriendsAndChatManager::SendGameInvite(const FUniqueNetId& ToUser)
 				}
 			}
 
-			Analytics.RecordGameInvite(*OnlineIdentity->GetUniquePlayerId(LocalControllerIndex), ToUser, TEXT("Social.GameInvite.Send"));
+			Analytics->FireEvent_SendGameInvite(*OnlineIdentity->GetUniquePlayerId(LocalControllerIndex), ToUser);
 		}
 	}
 }
@@ -2102,7 +2122,6 @@ void FFriendsAndChatManager::OnInviteAccepted(const FUniqueNetId& UserId, const
 		SendInviteAcceptedNotification(Friend);
 		Friend->SetPendingAccept();
 	}
-	RefreshList();
 	RequestListRefresh();
 }
 
@@ -2117,7 +2136,6 @@ void FFriendsAndChatManager::OnAcceptInviteComplete( int32 LocalPlayer, bool bWa
 	}
 	else
 	{
-		RefreshList();
 		RequestListRefresh();
 		SetState( EFriendsAndManagerState::Idle );
 	}
@@ -2127,8 +2145,6 @@ void FFriendsAndChatManager::OnDeleteFriendComplete( int32 LocalPlayer, bool bWa
 {
 	PendingOutgoingDeleteFriendRequests.Remove(DeletedFriendID.ToString());
 
-	RefreshList();
-
 	if ( PendingOutgoingDeleteFriendRequests.Num() > 0 )
 	{
 		SetState( EFriendsAndManagerState::DeletingFriends );
@@ -2150,137 +2166,4 @@ void FFriendsAndChatManager::AddFriendsToast(const FText Message)
 	}
 }
 
-
-// Analytics
-
-void FFriendsAndChatAnalytics::RecordGameInvite(const FUniqueNetId& LocalUserId, const FUniqueNetId& ToUser, const FString& EventStr) const
-{
-	if (Provider.IsValid())
-	{
-		if (LocalUserId.IsValid())
-		{
-			TArray<FAnalyticsEventAttribute> Attributes;
-			Attributes.Add(FAnalyticsEventAttribute(TEXT("User"), LocalUserId.ToString()));
-			Attributes.Add(FAnalyticsEventAttribute(TEXT("Friend"), ToUser.ToString()));
-			AddPresenceAttributes(LocalUserId, Attributes);
-			Provider->RecordEvent(EventStr, Attributes);
-		}
-	}
-}
-
-void FFriendsAndChatAnalytics::RecordFriendAction(const FUniqueNetId& LocalUserId, const IFriendItem& Friend, const FString& EventStr) const
-{
-	if (Provider.IsValid())
-	{
-		if (LocalUserId.IsValid())
-		{
-			TArray<FAnalyticsEventAttribute> Attributes;
-			Attributes.Add(FAnalyticsEventAttribute(TEXT("User"), LocalUserId.ToString()));
-			Attributes.Add(FAnalyticsEventAttribute(TEXT("Friend"), Friend.GetUniqueID()->ToString()));
-			AddPresenceAttributes(LocalUserId, Attributes);
-			Provider->RecordEvent(EventStr, Attributes);
-		}
-	}
-}
-
-void FFriendsAndChatAnalytics::RecordAddFriend(const FUniqueNetId& LocalUserId, const FString& FriendName, const FUniqueNetId& FriendId, EFindFriendResult::Type Result, bool bRecentPlayer, const FString& EventStr) const
-{
-	if (Provider.IsValid())
-	{
-		if (LocalUserId.IsValid())
-		{
-			TArray<FAnalyticsEventAttribute> Attributes;
-			Attributes.Add(FAnalyticsEventAttribute(TEXT("User"), LocalUserId.ToString()));
-			Attributes.Add(FAnalyticsEventAttribute(TEXT("Friend"), FriendId.ToString()));
-			Attributes.Add(FAnalyticsEventAttribute(TEXT("FriendName"), FriendName));
-			Attributes.Add(FAnalyticsEventAttribute(TEXT("Result"), EFindFriendResult::ToString(Result)));
-			Attributes.Add(FAnalyticsEventAttribute(TEXT("bRecentPlayer"), bRecentPlayer));
-			AddPresenceAttributes(LocalUserId, Attributes);
-			Provider->RecordEvent(EventStr, Attributes);
-		}
-	}
-}
-
-void FFriendsAndChatAnalytics::RecordToggleChat(const FUniqueNetId& LocalUserId, const FString& Channel, bool bEnabled, const FString& EventStr) const
-{
-	if (Provider.IsValid())
-	{
-		if (LocalUserId.IsValid())
-		{
-			TArray<FAnalyticsEventAttribute> Attributes;
-			Attributes.Add(FAnalyticsEventAttribute(TEXT("User"), LocalUserId.ToString()));
-			Attributes.Add(FAnalyticsEventAttribute(TEXT("Channel"), Channel));
-			Attributes.Add(FAnalyticsEventAttribute(TEXT("bEnabled"), bEnabled));
-			AddPresenceAttributes(LocalUserId, Attributes);
-			Provider->RecordEvent(EventStr, Attributes);
-		}
-	}
-}
-
-void FFriendsAndChatAnalytics::RecordPrivateChat(const FString& ToUser)
-{
-	int32& Count = PrivateChatCounts.FindOrAdd(ToUser);
-	Count += 1;
-}
-
-void FFriendsAndChatAnalytics::RecordChannelChat(const FString& ToChannel)
-{
-	int32& Count = ChannelChatCounts.FindOrAdd(ToChannel);
-	Count += 1;
-}
-
-void FFriendsAndChatAnalytics::FlushChatStats()
-{
-	if (Provider.IsValid())
-	{
-		IOnlineIdentityPtr OnlineIdentity = Online::GetIdentityInterface(TEXT("MCP"));
-		if (OnlineIdentity.IsValid())
-		{
-			TSharedPtr<const FUniqueNetId> UserId = OnlineIdentity->GetUniquePlayerId(0);
-			if (UserId.IsValid())
-			{
-				auto RecordSocialChatCountsEvents = [=](const TMap<FString, int32>& ChatCounts, const FString& ChatType)
-				{
-					if (ChatCounts.Num())
-					{
-						TArray<FAnalyticsEventAttribute> Attributes;
-						for (const auto& Pair : ChatCounts)
-						{
-							Attributes.Empty(3);
-							Attributes.Emplace(TEXT("Name"), Pair.Key);
-							Attributes.Emplace(TEXT("Type"), ChatType);
-							Attributes.Emplace(TEXT("Count"), Pair.Value);
-							Provider->RecordEvent("Social.Chat.Counts.2", Attributes);
-						}
-					}
-				};
-
-				RecordSocialChatCountsEvents(ChannelChatCounts, TEXT("Channel"));
-				RecordSocialChatCountsEvents(PrivateChatCounts, TEXT("Private"));
-			}
-		}
-	}
-	ChannelChatCounts.Empty();
-	PrivateChatCounts.Empty();
-}
-
-void FFriendsAndChatAnalytics::AddPresenceAttributes(const FUniqueNetId& UserId, TArray<FAnalyticsEventAttribute>& Attributes) const
-{
-	IOnlinePresencePtr OnlinePresence = Online::GetPresenceInterface(TEXT("MCP"));
-	if (OnlinePresence.IsValid())
-	{
-		TSharedPtr<FOnlineUserPresence> Presence;
-		OnlinePresence->GetCachedPresence(UserId, Presence);
-		if (Presence.IsValid())
-		{
-			FVariantData* ClientIdData = Presence->Status.Properties.Find(DefaultClientIdKey);
-			if (ClientIdData != nullptr)
-			{
-				Attributes.Add(FAnalyticsEventAttribute(TEXT("ClientId"), ClientIdData->ToString()));
-			}
-			Attributes.Add(FAnalyticsEventAttribute(TEXT("Status"), Presence->Status.StatusStr));
-		}
-	}
-}
-
 #undef LOCTEXT_NAMESPACE
diff --git a/Engine/Source/Developer/FriendsAndChat/Private/Layers/DataAccess/FriendsAndChatManager.h b/Engine/Source/Developer/FriendsAndChat/Private/Layers/DataAccess/FriendsAndChatManager.h
--- a/Engine/Source/Developer/FriendsAndChat/Private/Layers/DataAccess/FriendsAndChatManager.h
+++ b/Engine/Source/Developer/FriendsAndChat/Private/Layers/DataAccess/FriendsAndChatManager.h
@@ -56,55 +56,6 @@ namespace EFindFriendResult
 	}
 };
 
-// Analytics
-
-class FFriendsAndChatAnalytics
-{
-public:
-	FFriendsAndChatAnalytics() {}
-	/**
-	 * Update provider to use for capturing events
-	 */
-	void SetProvider(const TSharedPtr<IAnalyticsProvider>& AnalyticsProvider) {	Provider = AnalyticsProvider; }
-	/**
-	 * Record a game invite event
- 	 */
-	void RecordGameInvite(const FUniqueNetId& LocalUserId, const FUniqueNetId& ToUser, const FString& EventStr) const;
-	/**
-	 * Record a friend action event
-	 */
-	void RecordFriendAction(const FUniqueNetId& LocalUserId, const IFriendItem& Friend, const FString& EventStr) const;
-	/**
-	 * Record a add friend action event
-	 */
-	void RecordAddFriend(const FUniqueNetId& LocalUserId, const FString& FriendName, const FUniqueNetId& FriendId, EFindFriendResult::Type Result, bool bRecentPlayer, const FString& EventStr) const;
-	/**
-	 * Record chat option toggle
-	 */
-	void RecordToggleChat(const FUniqueNetId& LocalUserId, const FString& Channel, bool bEnabled, const FString& EventStr) const;
-	/**
-	 * Record a private chat to a user (aggregates pending FlushChat)
-	 */
-	void RecordPrivateChat(const FString& ToUser);
-	/**
-	 * Record a public chat to a channel (aggregates pending FlushChat)
-	 */
-	void RecordChannelChat(const FString& ToChannel);
-	/**
-	 * Flush any aggregated chat stats
-	 */
-	void FlushChatStats();
-
-private:
-	void AddPresenceAttributes(const FUniqueNetId& UserId, TArray<FAnalyticsEventAttribute>& Attributes) const;
-
-	// cached analytics provider for pushing events
-	TSharedPtr<IAnalyticsProvider> Provider;
-	// map of chat id to # of messages sent
-	TMap<FString, int32> PrivateChatCounts;
-	TMap<FString, int32> ChannelChatCounts;
-};
-
 /**
  * Implement the Friend and Chat manager
  */
@@ -186,7 +137,7 @@ public:
 	/**
 	 * Get the analytics for recording friends chat events
 	 */
-	FFriendsAndChatAnalytics& GetAnalytics()
+	TSharedPtr<class FFriendsAndChatAnalytics> GetAnalytics()
 	{
 		return Analytics;
 	}
@@ -320,7 +271,7 @@ public:
 	 * @param FriendItem The friend item to delete.
 	 * @param DeleteReason The reason the friend is being deleted.
 	 */
-	void DeleteFriend(TSharedPtr< IFriendItem > FriendItem, const FString& Action);
+	void DeleteFriend(TSharedPtr< IFriendItem > FriendItem, EFriendActionType::Type DeleteReason);
 
 	/**
 	 * Reject a game invite
@@ -732,6 +683,14 @@ private:
 	 */
 	FReply HandleMessageAccepted( TSharedPtr< FFriendsAndChatMessage > MessageNotification, EFriendsResponseType::Type ResponseType );
 
+	/**
+	 * Handle an join game message accepted from a notification.
+	 *
+	 * @param MessageNotification The message responded to.
+ 	 * @param ResponseTpe The response from the message request.
+	 */
+	FReply HandleGameInviteResponse( TSharedPtr< FFriendsAndChatMessage > MessageNotification, EFriendsResponseType::Type ResponseType );
+
 	/**
 	 * Validate friend list after a refresh. Check if the list has changed.
 	 *
@@ -990,7 +949,7 @@ private:
 
 private:
 
-	FFriendsAndChatAnalytics Analytics;
+	TSharedPtr<class FFriendsAndChatAnalytics> Analytics;
 	float FlushChatAnalyticsCountdown;
 
 	/** Handle to various registered delegates */
diff --git a/Engine/Source/Developer/FriendsAndChat/Private/Layers/DataAccess/FriendsMessageManager.cpp b/Engine/Source/Developer/FriendsAndChat/Private/Layers/DataAccess/FriendsMessageManager.cpp
--- a/Engine/Source/Developer/FriendsAndChat/Private/Layers/DataAccess/FriendsMessageManager.cpp
+++ b/Engine/Source/Developer/FriendsAndChat/Private/Layers/DataAccess/FriendsMessageManager.cpp
@@ -4,6 +4,7 @@
 #include "OnlineChatInterface.h"
 #include "ChatItemViewModel.h"
 #include "IChatCommunicationService.h"
+#include "FriendsAndChatAnalytics.h"
 
 #define LOCTEXT_NAMESPACE "FriendsMessageManager"
 // Message expiry time for different message types
@@ -117,14 +118,14 @@ public:
 				ChatItem->ToName = FText::FromString(*FoundFriend->GetName());
 				ChatItem->Message = MessageText;
 				ChatItem->MessageType = EChatMessageType::Whisper;
-				ChatItem->MessageTime = FDateTime::UtcNow();
-				ChatItem->ExpireTime = FDateTime::UtcNow() + FTimespan::FromSeconds(WhisperMessageLifetime);
+				ChatItem->MessageTime = FDateTime::Now();
+				ChatItem->ExpireTime = FDateTime::Now() + FTimespan::FromSeconds(WhisperMessageLifetime);
 				ChatItem->bIsFromSelf = true;
 				ChatItem->SenderId = LoggedInUser;
 				ChatItem->RecipientId = UserID;
 				AddMessage(ChatItem.ToSharedRef());
 				
-				FriendsAndChatManager.Pin()->GetAnalytics().RecordPrivateChat(UserID->ToString());
+				FriendsAndChatManager.Pin()->GetAnalytics()->RecordPrivateChat(UserID->ToString());
 
 				return ChatInterface->SendPrivateChat(*LoggedInUser, *UserID.Get(), MessageText.ToString());
 			}
@@ -138,8 +139,8 @@ public:
 		ChatItem->FromName = FText::FromString("Game");
 		ChatItem->Message = FText::FromString(MsgBody);
 		ChatItem->MessageType = EChatMessageType::Game;
-		ChatItem->MessageTime = FDateTime::UtcNow();
-		ChatItem->ExpireTime = FDateTime::UtcNow() + FTimespan::FromSeconds(GameMessageLifetime);
+		ChatItem->MessageTime = FDateTime::Now();
+		ChatItem->ExpireTime = FDateTime::Now() + FTimespan::FromSeconds(GameMessageLifetime);
 		ChatItem->bIsFromSelf = false;
 		GameMessagesCount++;
 		AddMessage(ChatItem.ToSharedRef());
@@ -281,8 +282,8 @@ private:
 			}
 			ChatItem->Message = FText::FromString(TEXT("entered room"));
 			ChatItem->MessageType = EChatMessageType::Global;
-			ChatItem->MessageTime = FDateTime::UtcNow();
-			ChatItem->ExpireTime = FDateTime::UtcNow() + GlobalMessageLifetime;
+			ChatItem->MessageTime = FDateTime::Now();
+			ChatItem->ExpireTime = FDateTime::Now() + GlobalMessageLifetime;
 			ChatItem->bIsFromSelf = false;
 			GlobalMessagesCount++;
 			AddMessage(ChatItem.ToSharedRef());
@@ -303,8 +304,8 @@ private:
 			}
 			ChatItem->Message = FText::FromString(TEXT("left room"));
 			ChatItem->MessageType = EChatMessageType::Global;
-			ChatItem->MessageTime = FDateTime::UtcNow();
-			ChatItem->ExpireTime = FDateTime::UtcNow() + GlobalMessageLifetime;
+			ChatItem->MessageTime = FDateTime::Now();
+			ChatItem->ExpireTime = FDateTime::Now() + GlobalMessageLifetime;
 			ChatItem->bIsFromSelf = false;
 			GlobalMessagesCount++;
 			AddMessage(ChatItem.ToSharedRef());
@@ -338,8 +339,8 @@ private:
 		}
 		ChatItem->FromName = FText::FromString(*ChatMessage->GetNickname());
 		ChatItem->Message = FText::FromString(*ChatMessage->GetBody());
-		ChatItem->MessageTime = ChatMessage->GetTimestamp();
-		ChatItem->ExpireTime = ChatMessage->GetTimestamp() + FTimespan::FromSeconds(GlobalMessageLifetime);
+		ChatItem->MessageTime = FDateTime::Now();
+		ChatItem->ExpireTime = FDateTime::Now() + FTimespan::FromSeconds(GlobalMessageLifetime);
 		ChatItem->bIsFromSelf = *ChatMessage->GetUserId() == *LoggedInUser;
 		ChatItem->SenderId = ChatMessage->GetUserId();
 		ChatItem->RecipientId = nullptr;
@@ -361,8 +362,8 @@ private:
 			ChatItem->RecipientId = LoggedInUser;
 			ChatItem->Message = FText::FromString(*ChatMessage->GetBody());
 			ChatItem->MessageType = EChatMessageType::Whisper;
-			ChatItem->MessageTime = ChatMessage->GetTimestamp();
-			ChatItem->ExpireTime = ChatMessage->GetTimestamp() + FTimespan::FromSeconds(WhisperMessageLifetime);
+			ChatItem->MessageTime = FDateTime::Now();
+			ChatItem->ExpireTime = FDateTime::Now() + FTimespan::FromSeconds(WhisperMessageLifetime);
 			ChatItem->MessageRef = ChatMessage;
 			WhisperMessagesCount++;
 			AddMessage(ChatItem.ToSharedRef());
@@ -397,7 +398,7 @@ private:
 			bool bGameTimeFound = false;
 			bool bPartyTimeFound = false;
 			bool bWhisperFound = false;
-			FDateTime CurrentTime = FDateTime::UtcNow();
+			FDateTime CurrentTime = FDateTime::Now();
 			for(int32 Index = 0; Index < ReceivedMessages.Num(); Index++)
 			{
 				TSharedRef<FFriendChatMessage> Message = ReceivedMessages[Index];
diff --git a/Engine/Source/Developer/FriendsAndChat/Private/Layers/Domain/FriendsAndChatAnalytics.cpp b/Engine/Source/Developer/FriendsAndChat/Private/Layers/Domain/FriendsAndChatAnalytics.cpp
new file mode 100644
--- /dev/null
+++ b/Engine/Source/Developer/FriendsAndChat/Private/Layers/Domain/FriendsAndChatAnalytics.cpp
@@ -0,0 +1,227 @@
+// Copyright 1998-2015 Epic Games, Inc. All Rights Reserved.
+
+#include "FriendsAndChatPrivatePCH.h"
+#include "FriendsAndChatAnalytics.h"
+
+class FFriendsAndChatAnalyticsImpl
+	: public FFriendsAndChatAnalytics
+{
+public:
+
+	virtual ~FFriendsAndChatAnalyticsImpl() { }
+
+	virtual void SetProvider(const TSharedPtr<IAnalyticsProvider>& InAnalyticsProvider) override
+	{
+		AnalyticsProvider = InAnalyticsProvider;
+	}
+
+	virtual void FireEvent_AcceptGameInvite(const FUniqueNetId& LocalUserId, const FUniqueNetId& ToUser) const override
+	{
+		static const FString EventName = TEXT("Social.GameInvite.Accept");
+
+		if (AnalyticsProvider.IsValid())
+		{
+			if (LocalUserId.IsValid())
+			{
+				TArray<FAnalyticsEventAttribute> Attributes;
+				Attributes.Add(FAnalyticsEventAttribute(TEXT("User"), LocalUserId.ToString()));
+				Attributes.Add(FAnalyticsEventAttribute(TEXT("Friend"), ToUser.ToString()));
+				AddPresenceAttributes(LocalUserId, Attributes);
+				AnalyticsProvider->RecordEvent(EventName, Attributes);
+			}
+		}
+	}
+
+	virtual void FireEvent_RejectGameInvite(const FUniqueNetId& LocalUserId, const FUniqueNetId& ToUser) const override
+	{
+
+		static const FString EventName =  TEXT("Social.GameInvite.Reject");
+		if (AnalyticsProvider.IsValid())
+		{
+			if (LocalUserId.IsValid())
+			{
+				TArray<FAnalyticsEventAttribute> Attributes;
+				Attributes.Add(FAnalyticsEventAttribute(TEXT("User"), LocalUserId.ToString()));
+				Attributes.Add(FAnalyticsEventAttribute(TEXT("Friend"), ToUser.ToString()));
+				AddPresenceAttributes(LocalUserId, Attributes);
+				AnalyticsProvider->RecordEvent(EventName, Attributes);
+			}
+		}
+	}
+
+	virtual void FireEvent_SendGameInvite(const FUniqueNetId& LocalUserId, const FUniqueNetId& ToUser) const override
+	{
+		static const FString EventName = TEXT("Social.GameInvite.Send");
+
+		if (AnalyticsProvider.IsValid())
+		{
+			if (LocalUserId.IsValid())
+			{
+				TArray<FAnalyticsEventAttribute> Attributes;
+				Attributes.Add(FAnalyticsEventAttribute(TEXT("User"), LocalUserId.ToString()));
+				Attributes.Add(FAnalyticsEventAttribute(TEXT("Friend"), ToUser.ToString()));
+				AddPresenceAttributes(LocalUserId, Attributes);
+				AnalyticsProvider->RecordEvent(EventName, Attributes);
+			}
+		}
+	}
+
+	virtual void FireEvent_FriendRequestAccepted(const FUniqueNetId& LocalUserId, const IFriendItem& Friend) const override
+	{
+		static const FString EventName = TEXT("Social.FriendAction.Accept");
+	
+		if (AnalyticsProvider.IsValid())
+		{
+			if (LocalUserId.IsValid())
+			{
+				TArray<FAnalyticsEventAttribute> Attributes;
+				Attributes.Add(FAnalyticsEventAttribute(TEXT("User"), LocalUserId.ToString()));
+				Attributes.Add(FAnalyticsEventAttribute(TEXT("Friend"), Friend.GetUniqueID()->ToString()));
+				AddPresenceAttributes(LocalUserId, Attributes);
+				AnalyticsProvider->RecordEvent(EventName, Attributes);
+			}
+		}
+	}
+
+	virtual void FireEvent_FriendRemoved(const FUniqueNetId& LocalUserId, const IFriendItem& Friend, const FString& RemoveReason) const override
+	{
+		static const FString EventName = TEXT("Social.FriendRemoved");
+		if (AnalyticsProvider.IsValid())
+		{
+			if (LocalUserId.IsValid())
+			{
+				TArray<FAnalyticsEventAttribute> Attributes;
+				Attributes.Add(FAnalyticsEventAttribute(TEXT("User"), LocalUserId.ToString()));
+				Attributes.Add(FAnalyticsEventAttribute(TEXT("Friend"), Friend.GetUniqueID()->ToString()));
+				Attributes.Add(FAnalyticsEventAttribute(TEXT("RemoveReason"), RemoveReason));
+				AddPresenceAttributes(LocalUserId, Attributes);
+				AnalyticsProvider->RecordEvent(EventName, Attributes);
+			}
+		}
+	}
+
+	virtual void FireEvent_AddFriend(const FUniqueNetId& LocalUserId, const FString& FriendName, const FUniqueNetId& FriendId, EFindFriendResult::Type Result, bool bRecentPlayer) const override
+	{
+		static const FString EventName = TEXT("Social.AddFriend");
+
+		if (AnalyticsProvider.IsValid())
+		{
+			if (LocalUserId.IsValid())
+			{
+				TArray<FAnalyticsEventAttribute> Attributes;
+				Attributes.Add(FAnalyticsEventAttribute(TEXT("User"), LocalUserId.ToString()));
+				Attributes.Add(FAnalyticsEventAttribute(TEXT("Friend"), FriendId.ToString()));
+				Attributes.Add(FAnalyticsEventAttribute(TEXT("FriendName"), FriendName));
+				Attributes.Add(FAnalyticsEventAttribute(TEXT("Result"), EFindFriendResult::ToString(Result)));
+				Attributes.Add(FAnalyticsEventAttribute(TEXT("bRecentPlayer"), bRecentPlayer));
+				AddPresenceAttributes(LocalUserId, Attributes);
+				AnalyticsProvider->RecordEvent(EventName, Attributes);
+			}
+		}
+	}
+
+	virtual void FireEvent_RecordToggleChat(const FUniqueNetId& LocalUserId, const FString& Channel, bool bEnabled) const override
+	{
+		static const FString EventName = TEXT("Social.Chat.Toggle");
+		if (AnalyticsProvider.IsValid())
+		{
+			if (LocalUserId.IsValid())
+			{
+				TArray<FAnalyticsEventAttribute> Attributes;
+				Attributes.Add(FAnalyticsEventAttribute(TEXT("User"), LocalUserId.ToString()));
+				Attributes.Add(FAnalyticsEventAttribute(TEXT("Channel"), Channel));
+				Attributes.Add(FAnalyticsEventAttribute(TEXT("bEnabled"), bEnabled));
+				AddPresenceAttributes(LocalUserId, Attributes);
+				AnalyticsProvider->RecordEvent(EventName, Attributes);
+			}
+		}
+	}
+
+	virtual void RecordPrivateChat(const FString& ToUser) override
+	{
+		int32& Count = PrivateChatCounts.FindOrAdd(ToUser);
+		Count += 1;
+	}
+
+	virtual void RecordChannelChat(const FString& ToChannel) override
+	{
+		int32& Count = ChannelChatCounts.FindOrAdd(ToChannel);
+		Count += 1;
+	}
+
+	virtual void FireEvent_FlushChatStats() override
+	{
+		if (AnalyticsProvider.IsValid())
+		{
+			IOnlineIdentityPtr OnlineIdentity = Online::GetIdentityInterface(TEXT("MCP"));
+			if (OnlineIdentity.IsValid())
+			{
+				TSharedPtr<const FUniqueNetId> UserId = OnlineIdentity->GetUniquePlayerId(0);
+				if (UserId.IsValid())
+				{
+					auto RecordSocialChatCountsEvents = [=](const TMap<FString, int32>& ChatCounts, const FString& ChatType)
+					{
+						if (ChatCounts.Num())
+						{
+							TArray<FAnalyticsEventAttribute> Attributes;
+							for (const auto& Pair : ChatCounts)
+							{
+								Attributes.Empty(3);
+								Attributes.Emplace(TEXT("Name"), Pair.Key);
+								Attributes.Emplace(TEXT("Type"), ChatType);
+								Attributes.Emplace(TEXT("Count"), Pair.Value);
+								AnalyticsProvider->RecordEvent("Social.Chat.Counts.2", Attributes);
+							}
+						}
+					};
+
+					RecordSocialChatCountsEvents(ChannelChatCounts, TEXT("Channel"));
+					RecordSocialChatCountsEvents(PrivateChatCounts, TEXT("Private"));
+				}
+			}
+		}
+		ChannelChatCounts.Empty();
+		PrivateChatCounts.Empty();
+	}
+
+private: // functions
+
+	void AddPresenceAttributes(const FUniqueNetId& UserId, TArray<FAnalyticsEventAttribute>& Attributes) const
+	{
+		IOnlinePresencePtr OnlinePresence = Online::GetPresenceInterface(TEXT("MCP"));
+		if (OnlinePresence.IsValid())
+		{
+			TSharedPtr<FOnlineUserPresence> Presence;
+			OnlinePresence->GetCachedPresence(UserId, Presence);
+			if (Presence.IsValid())
+			{
+				FVariantData* ClientIdData = Presence->Status.Properties.Find(DefaultClientIdKey);
+				if (ClientIdData != nullptr)
+				{
+					Attributes.Add(FAnalyticsEventAttribute(TEXT("ClientId"), ClientIdData->ToString()));
+				}
+				Attributes.Add(FAnalyticsEventAttribute(TEXT("Status"), Presence->Status.StatusStr));
+			}
+		}
+	}
+
+	FFriendsAndChatAnalyticsImpl()
+	{}
+
+private: // data
+
+	// cached analytics provider for pushing events
+	TSharedPtr<IAnalyticsProvider> AnalyticsProvider;
+
+	// map of chat id to # of messages sent
+	TMap<FString, int32> PrivateChatCounts;
+	TMap<FString, int32> ChannelChatCounts;
+
+	friend FFriendsAndChatAnalyticsFactory;
+};
+
+TSharedRef< FFriendsAndChatAnalytics > FFriendsAndChatAnalyticsFactory::Create()
+{
+	TSharedRef< FFriendsAndChatAnalyticsImpl > Instance = MakeShareable(new FFriendsAndChatAnalyticsImpl());
+	return Instance;
+}
diff --git a/Engine/Source/Developer/FriendsAndChat/Private/Layers/Domain/FriendsChatMarkupService.cpp b/Engine/Source/Developer/FriendsAndChat/Private/Layers/Domain/FriendsChatMarkupService.cpp
--- a/Engine/Source/Developer/FriendsAndChat/Private/Layers/Domain/FriendsChatMarkupService.cpp
+++ b/Engine/Source/Developer/FriendsAndChat/Private/Layers/Domain/FriendsChatMarkupService.cpp
@@ -114,6 +114,36 @@ private:
 	TSharedRef<FFriendViewModel> FriendItem;
 };
 
+class FInvalidChatTip : public IChatTip
+{
+public:
+	virtual FText GetTipText() override
+	{
+		return LOCTEXT("InvalidChatTip", "Invalid markup");
+	}
+
+	virtual bool IsValidForType(EChatMessageType::Type ChatChannel) override
+	{
+		return true;
+	}
+
+	virtual bool IsEnabled() override
+	{
+		return true;
+	}
+
+	virtual FReply ExecuteTip() override
+	{
+		return FReply::Handled();
+	}
+
+	FInvalidChatTip()
+	{}
+
+	virtual ~FInvalidChatTip(){}
+};
+
+
 class FFriendsChatMarkupServiceImpl
 	: public FFriendsChatMarkupService
 {
@@ -176,6 +206,10 @@ public:
 			{
 				NavigationService->ChangeChatChannel(EChatMessageType::Custom);
 			}
+			else if (ValidatedInput->ChatChannel == EChatMessageType::Empty)
+			{
+				NavigationService->ChangeChatChannel(EChatMessageType::Empty);
+			}
 			else if(ValidatedInput->ChatChannel == EChatMessageType::Invalid)
 			{
 				// Handle bad text
@@ -255,7 +289,7 @@ public:
 
 			Args.Add(TEXT("NameStyle"), FText::FromString(HyperlinkStyle));
 
-			FText ValidatedText = FText::Format(NSLOCTEXT("FChatItemViewModel", "DisplayNameAndMessage", "/w <a id=\"UserName\" uid=\"{UniqueID}\" Username=\"{SenderName}\" style=\"{NameStyle}\">{SenderName}</> "), Args);
+			FText ValidatedText = FText::Format(NSLOCTEXT("FFriendsChatMarkupService", "DisplayNameAndMessage", "/w <a id=\"UserName\" uid=\"{UniqueID}\" Username=\"{SenderName}\" style=\"{NameStyle}\">{SenderName}</> "), Args);
 			ValidatedInputString = ValidatedText.ToString();
 			SelectedFriendViewModel = FriendViewModel;
 			ProcessedInput->ValidFriends.Empty();
@@ -331,16 +365,20 @@ private:
 				ChatTipArray.Add(MakeShareable(new FFriendChatTip(SharedThis(this), Friend.ToSharedRef())));
 			}
 		}
-		else
+		else if(ProcessedInput->ChatChannel != EChatMessageType::Invalid)
 		{
 			for(const auto& ChatTip : CommmonChatTips)
 			{
-				if(ProcessedInput->ChatChannel == EChatMessageType::Invalid || ChatTip->IsValidForType(ProcessedInput->ChatChannel))
+				if(ProcessedInput->ChatChannel == EChatMessageType::Empty || ChatTip->IsValidForType(ProcessedInput->ChatChannel))
 				{
 					ChatTipArray.Add(ChatTip); 
 				}
 			}
 		}
+		else
+		{
+			ChatTipArray.Add(MakeShareable(new FInvalidChatTip()));
+		}
 
 		if(ChatTipArray.Num() > 0)
 		{
@@ -444,7 +482,7 @@ private:
 						}
 					}
 				}
-				else
+				else if(ChatType != EChatMessageType::Empty)
 				{
 					ProcessedInput->NeedsTip = false;
 					ProcessedInput->Message = InInputText.RightChop(NavigationToken.Len() + 1);
diff --git a/Engine/Source/Developer/FriendsAndChat/Private/Layers/Presentation/ChatItemViewModel.cpp b/Engine/Source/Developer/FriendsAndChat/Private/Layers/Presentation/ChatItemViewModel.cpp
--- a/Engine/Source/Developer/FriendsAndChat/Private/Layers/Presentation/ChatItemViewModel.cpp
+++ b/Engine/Source/Developer/FriendsAndChat/Private/Layers/Presentation/ChatItemViewModel.cpp
@@ -19,16 +19,6 @@ public:
 		return ChatMessages[0]->MessageType;
 	}
 
-	virtual FText GetMessageTimeText() const override
-	{
-		return MessageTimeAsText;
-	}
-
-	virtual FText GetSimpleMessageTimeText() const override
-	{
-		return FText::AsTime(ChatMessages[0]->MessageTime, EDateTimeStyle::Short);
-	}
-
 	virtual FDateTime GetMessageTime() const override
 	{
 		return ChatMessages.Last()->MessageTime;
@@ -38,9 +28,6 @@ public:
 	{
 		ChatMessages.Add(ChatMessage);
 
-		//Display the time of the latest message only
-		FTicker::GetCoreTicker().RemoveTicker(TickerHandle);
-
 		FTextBuilder MessageBuilder;
 		for (int32 Index = 0; Index < ChatMessages.Num(); Index++)
 		{
@@ -54,11 +41,6 @@ public:
 		}
 
 		ComposedMessage = MessageBuilder.ToText();
-
-		if (UpdateTimeStamp(60)) //this will broadcast a changed event
-		{
-			TickerHandle = FTicker::GetCoreTicker().AddTicker(FTickerDelegate::CreateSP(this, &FChatItemViewModelImpl::UpdateTimeStamp), 60);
-		}
 	}
 
 	virtual FText GetSenderName() const override
@@ -86,75 +68,16 @@ public:
 		return ChatMessages[0]->bIsFromSelf;
 	}
 
-	DECLARE_DERIVED_EVENT(FChatItemViewModelImpl, FChatItemViewModel::FChangedEvent, FChangedEvent)
-	virtual FChatItemViewModel::FChangedEvent& OnChanged() override { return ChangedEvent; }
-
-	virtual ~FChatItemViewModelImpl()
-	{
-		FTicker::GetCoreTicker().RemoveTicker(TickerHandle);
-	}
-
 private:
 
 	FChatItemViewModelImpl()
 	{}
 
-	bool UpdateTimeStamp(float Delay)
-	{
-		bool ContinueTicking = true;
-
-		TSharedRef<FFriendChatMessage> Message = ChatMessages.Last();
-		FDateTime Now = FDateTime::UtcNow();
-
-		if ((Now - Message->MessageTime).GetDuration() <= FTimespan::FromMinutes(1.0))
-		{
-			static FText NowTimeStamp = NSLOCTEXT("SChatWindow", "Now_TimeStamp", "Now");
-			MessageTimeAsText = NowTimeStamp;
-		}
-		else if ((Now - Message->MessageTime).GetDuration() <= FTimespan::FromMinutes(2.0))
-		{
-			static FText OneMinuteTimeStamp = NSLOCTEXT("SChatWindow", "1_Minute_TimeStamp", "1 min");
-			MessageTimeAsText = OneMinuteTimeStamp;
-		}
-		else if ((Now - Message->MessageTime).GetDuration() <= FTimespan::FromMinutes(3.0))
-		{
-			static FText TwoMinuteTimeStamp = NSLOCTEXT("SChatWindow", "2_Minute_TimeStamp", "2 min");
-			MessageTimeAsText = TwoMinuteTimeStamp;
-		}
-		else if ((Now - Message->MessageTime).GetDuration() <= FTimespan::FromMinutes(4.0))
-		{
-			static FText ThreeMinuteTimeStamp = NSLOCTEXT("SChatWindow", "3_Minute_TimeStamp", "3 min");
-			MessageTimeAsText = ThreeMinuteTimeStamp;
-		}
-		else if ((Now - Message->MessageTime).GetDuration() <= FTimespan::FromMinutes(5.0))
-		{
-			static FText FourMinuteTimeStamp = NSLOCTEXT("SChatWindow", "4_Minute_TimeStamp", "4 min");
-			MessageTimeAsText = FourMinuteTimeStamp;
-		}
-		else if ((Now - Message->MessageTime).GetDuration() <= FTimespan::FromMinutes(6.0))
-		{
-			static FText FiveMinuteTimeStamp = NSLOCTEXT("SChatWindow", "5_Minute_TimeStamp", "5 min");
-			MessageTimeAsText = FiveMinuteTimeStamp;
-		}
-		else
-		{
-			MessageTimeAsText = FText::AsTime(ChatMessages.Last()->MessageTime, EDateTimeStyle::Short);
-			ContinueTicking = false;
-		}
-
-		OnChanged().Broadcast(SharedThis(this));
-		return ContinueTicking;
-	}
-
 private:
 
 	TArray<TSharedRef<FFriendChatMessage>> ChatMessages;
 
 	FText ComposedMessage;
-	FText MessageTimeAsText;
-	FDelegateHandle TickerHandle;
-
-	FChangedEvent ChangedEvent;
 
 	friend FChatItemViewModelFactory;
 };
diff --git a/Engine/Source/Developer/FriendsAndChat/Private/Layers/Presentation/ChatItemViewModel.h b/Engine/Source/Developer/FriendsAndChat/Private/Layers/Presentation/ChatItemViewModel.h
--- a/Engine/Source/Developer/FriendsAndChat/Private/Layers/Presentation/ChatItemViewModel.h
+++ b/Engine/Source/Developer/FriendsAndChat/Private/Layers/Presentation/ChatItemViewModel.h
@@ -13,8 +13,6 @@ public:
 
 	virtual FText GetMessage() const = 0;
 	virtual const EChatMessageType::Type GetMessageType() const =0;
-	virtual FText GetMessageTimeText() const = 0;
-	virtual FText GetSimpleMessageTimeText() const = 0;
 	virtual FDateTime GetMessageTime() const = 0;
 
 	virtual void AddMessage(const TSharedRef<FFriendChatMessage>& ChatMessage) = 0;
@@ -26,9 +24,6 @@ public:
 	virtual const TSharedPtr<const FUniqueNetId> GetRecipientID() const = 0;
 
 	virtual const bool IsFromSelf() const = 0;
-
-	DECLARE_EVENT_OneParam(FChatItemViewModel, FChangedEvent, const TSharedRef<FChatItemViewModel>&)
-	virtual FChangedEvent& OnChanged() = 0;
 };
 
 /**
diff --git a/Engine/Source/Developer/FriendsAndChat/Private/Layers/Presentation/ChatViewModel.cpp b/Engine/Source/Developer/FriendsAndChat/Private/Layers/Presentation/ChatViewModel.cpp
--- a/Engine/Source/Developer/FriendsAndChat/Private/Layers/Presentation/ChatViewModel.cpp
+++ b/Engine/Source/Developer/FriendsAndChat/Private/Layers/Presentation/ChatViewModel.cpp
@@ -10,6 +10,7 @@
 #include "FriendsChatMarkupService.h"
 #include "ChatTipViewModel.h"
 #include "ChatSettingsService.h"
+#include "FriendsAndChatAnalytics.h"
 
 #define LOCTEXT_NAMESPACE ""
 
@@ -145,7 +146,7 @@ public:
 			}
 			else if(SelectedFriend.IsValid())
 			{
-				return FText::Format(NSLOCTEXT("ChatViewModel", "FriendUnavailable", "Unable to send whisper to {0}."), SelectedFriend->DisplayName);
+				return FText::Format(NSLOCTEXT("ChatViewModel", "FriendUnavailable_WithDisplayName", "Unable to send whisper to {0}."), SelectedFriend->DisplayName);
 			}
 			else
 			{
@@ -186,7 +187,7 @@ public:
 								//@todo will need to support multiple party channels eventually, hardcoded to first party for now
 								bSuccess = MessageManager.Pin()->SendRoomMessage((*PartyRoomId).ToString(), NewMessage.ToString());
 
-								FriendsAndChatManager.Pin()->GetAnalytics().RecordChannelChat(TEXT("Party"));
+								FriendsAndChatManager.Pin()->GetAnalytics()->RecordChannelChat(TEXT("Party"));
 							}
 						}
 						break;
@@ -194,7 +195,7 @@ public:
 						{
 								//@todo samz - send message to specific room (empty room name will send to all rooms)
 								bSuccess = MessageManager.Pin()->SendRoomMessage(FString(), PlainText.ToString());
-								FriendsAndChatManager.Pin()->GetAnalytics().RecordChannelChat(TEXT("Global"));
+								FriendsAndChatManager.Pin()->GetAnalytics()->RecordChannelChat(TEXT("Global"));
 						}
 					}
 				}
diff --git a/Engine/Source/Developer/FriendsAndChat/Private/Layers/Presentation/ClanCollectionViewModel.cpp b/Engine/Source/Developer/FriendsAndChat/Private/Layers/Presentation/ClanCollectionViewModel.cpp
new file mode 100644
--- /dev/null
+++ b/Engine/Source/Developer/FriendsAndChat/Private/Layers/Presentation/ClanCollectionViewModel.cpp
@@ -0,0 +1,96 @@
+// Copyright 1998-2015 Epic Games, Inc. All Rights Reserved.
+
+#include "FriendsAndChatPrivatePCH.h"
+#include "ClanCollectionViewModel.h"
+#include "ClanInfoViewModel.h"
+#include "ClanRepository.h"
+#include "ClanInfo.h"
+#include "IFriendList.h"
+
+class FClanCollectionViewModelImpl
+	: public FClanCollectionViewModel
+{
+public:
+
+	virtual FText GetClanTitle() override
+	{
+		if(ClanList.Num() > 0)
+		{
+			return ClanList[0]->GetClanTitle();
+		}
+		return FText::GetEmpty();
+	}
+
+	virtual void EnumerateJoinedClans(TArray<TSharedRef<FClanInfoViewModel>>& OUTClanList) override
+	{
+		for (const auto& ClanItem : ClanList)
+		{
+			OUTClanList.Add(ClanItem);
+		}
+	}
+
+	virtual int32 NumberOfJoinedClans() const override
+	{
+		return ClanList.Num();
+	}
+
+	virtual int32 NumberOfClanInvites() const override
+	{
+		return ClanList.Num();
+	}
+
+
+	DECLARE_DERIVED_EVENT(FClanCollectionViewModelImpl, FClanCollectionViewModel::FOpenClanDetails, FOpenClanDetails);
+	virtual FOpenClanDetails& OpenClanDetails() override
+	{
+		return OpenClanDetailsEvent;
+	}
+
+private:
+	void Initialize()
+	{
+		TArray<TSharedRef<IClanInfo>> ClanInfos;
+		ClanRepository->EnumerateClanList(ClanInfos);
+		for (const auto& ClanItem : ClanInfos)
+		{
+			TSharedRef<FClanInfoViewModel> CIModel = FClanInfoViewModelFactory::Create(ClanItem, FriendsListFactory, FriendsAndChatManager.Pin().ToSharedRef());
+			CIModel->OpenClanDetails().AddSP(this, &FClanCollectionViewModelImpl::OnOpenClanDetails);
+			ClanList.Add(CIModel);
+		}
+	}
+
+	void OnOpenClanDetails(const TSharedRef<FClanInfoViewModel > &ClanInfoViewModel)
+	{
+		OpenClanDetailsEvent.Broadcast(ClanInfoViewModel);
+	}
+
+private:
+	FClanCollectionViewModelImpl(
+	const TSharedRef<IClanRepository>& InClanRepository,
+	const TSharedRef<IFriendListFactory>& InFriendsListFactory,
+	const TSharedRef<FFriendsAndChatManager>& InFriendsAndChatManager)
+		: ClanRepository(InClanRepository)
+		, FriendsListFactory(InFriendsListFactory)
+		, FriendsAndChatManager(InFriendsAndChatManager)
+	{
+	}
+
+	FOpenClanDetails OpenClanDetailsEvent;
+	TSharedRef<IClanRepository> ClanRepository;
+	TArray<TSharedRef<FClanInfoViewModel>> ClanList;
+	TSharedRef<IFriendListFactory> FriendsListFactory;
+	TWeakPtr<FFriendsAndChatManager> FriendsAndChatManager;
+
+	friend FClanCollectionViewModelFactory;
+};
+
+TSharedRef< FClanCollectionViewModel > FClanCollectionViewModelFactory::Create(
+	const TSharedRef<IClanRepository>& ClanRepository,
+	const TSharedRef<IFriendListFactory>& FriendsListFactory,
+	const TSharedRef<FFriendsAndChatManager>& FriendsAndChatManager
+	)
+{
+	TSharedRef< FClanCollectionViewModelImpl > ViewModel(new FClanCollectionViewModelImpl(ClanRepository, FriendsListFactory, FriendsAndChatManager));
+	ViewModel->Initialize();
+	return ViewModel;
+}
\ No newline at end of file
diff --git a/Engine/Source/Developer/FriendsAndChat/Private/Layers/Presentation/ClanCollectionViewModel.h b/Engine/Source/Developer/FriendsAndChat/Private/Layers/Presentation/ClanCollectionViewModel.h
new file mode 100644
--- /dev/null
+++ b/Engine/Source/Developer/FriendsAndChat/Private/Layers/Presentation/ClanCollectionViewModel.h
@@ -0,0 +1,31 @@
+// Copyright 1998-2015 Epic Games, Inc. All Rights Reserved.
+
+#pragma once
+
+class FClanCollectionViewModel
+	: public TSharedFromThis<FClanCollectionViewModel>
+{
+public:
+	virtual ~FClanCollectionViewModel() {}
+
+	virtual FText GetClanTitle() = 0;
+
+	virtual void EnumerateJoinedClans(TArray<TSharedRef<class FClanInfoViewModel>>& OUTClanList) = 0;
+
+	virtual int32 NumberOfJoinedClans() const = 0;
+
+	virtual int32 NumberOfClanInvites() const = 0;
+
+	DECLARE_EVENT_OneParam(FClanCollectionViewModel, FOpenClanDetails, const TSharedRef<class FClanInfoViewModel>&)
+	virtual FOpenClanDetails& OpenClanDetails() = 0;
+};
+
+/**
+ * Creates the implementation for an ClanViewModel.
+ *
+ * @return the newly created ClanViewModel implementation.
+ */
+FACTORY(TSharedRef< FClanCollectionViewModel >, FClanCollectionViewModel,
+	const TSharedRef<class IClanRepository>& ClanRepository,
+	const TSharedRef<class IFriendListFactory>& FriendsListFactory,
+	const TSharedRef<class FFriendsAndChatManager>& FriendsAndChatManager);
\ No newline at end of file
diff --git a/Engine/Source/Developer/FriendsAndChat/Private/Layers/Presentation/ClanInfoViewModel.cpp b/Engine/Source/Developer/FriendsAndChat/Private/Layers/Presentation/ClanInfoViewModel.cpp
--- a/Engine/Source/Developer/FriendsAndChat/Private/Layers/Presentation/ClanInfoViewModel.cpp
+++ b/Engine/Source/Developer/FriendsAndChat/Private/Layers/Presentation/ClanInfoViewModel.cpp
@@ -37,6 +37,22 @@ public:
 		return ClanInfo->GetClanBrushName();
 	}
 
+	virtual void EnumerateActions(TArray<EClanActionType::Type>& OUTActions) override
+	{
+		// TODO Antony.Carter - Add enumerations on a context specific basis
+		OUTActions.Add(EClanActionType::AcceptClanRequest);
+		OUTActions.Add(EClanActionType::LeaveClan);
+		OUTActions.Add(EClanActionType::CancelClanRequest);
+		OUTActions.Add(EClanActionType::IgnoreClanRequest);
+		OUTActions.Add(EClanActionType::SetPrimaryClan);
+	}
+
+	DECLARE_DERIVED_EVENT(FClanInfoViewModelImpl, FClanInfoViewModel::FOpenClanDetails, FOpenClanDetails);
+	virtual FOpenClanDetails& OpenClanDetails() override
+	{
+		return OpenClanDetailsEvent;
+	}
+
 private:
 	void Initialize()
 	{
@@ -54,6 +70,7 @@ private:
 	{
 	}
 
+	FOpenClanDetails OpenClanDetailsEvent;
 	TSharedRef<IClanInfo> ClanInfo;
 	TSharedRef<IFriendListFactory> FriendsListFactory;
 	TWeakPtr<FFriendsAndChatManager> FriendsAndChatManager;
diff --git a/Engine/Source/Developer/FriendsAndChat/Private/Layers/Presentation/ClanInfoViewModel.h b/Engine/Source/Developer/FriendsAndChat/Private/Layers/Presentation/ClanInfoViewModel.h
--- a/Engine/Source/Developer/FriendsAndChat/Private/Layers/Presentation/ClanInfoViewModel.h
+++ b/Engine/Source/Developer/FriendsAndChat/Private/Layers/Presentation/ClanInfoViewModel.h
@@ -2,6 +2,60 @@
 
 #pragma once
 
+
+#define LOCTEXT_NAMESPACE ""
+
+enum class EClanActionLevel
+{
+	Action,			// A normal action
+	Emphasis,		// An important action to be shown with emphasis
+	Critical,		// A critical action
+};
+
+
+// Enum holding friend action types
+namespace EClanActionType
+{
+	enum Type : uint8
+	{
+		AcceptClanRequest,		// Accept an incoming clan request
+		IgnoreClanRequest,		// Ignore an incoming clan request
+		SendClanRequest,		// Send an outgoing clan request
+		CancelClanRequest,		// Cancel an outgoing clan request
+		LeaveClan,				// Leave Clan
+		SetPrimaryClan,			// Set as Primary Clan
+		MAX_None,
+	};
+
+	inline FText ToText(EClanActionType::Type State)
+	{
+		switch (State)
+		{
+		case AcceptClanRequest: return LOCTEXT("AcceptClanRequest", "Accept");
+		case IgnoreClanRequest: return LOCTEXT("IgnoreClanRequest", "Ignore");
+		case CancelClanRequest: return LOCTEXT("RejectClanRequest", "Reject");
+		case LeaveClan: return LOCTEXT("CancelClanRequest", "Cancel");
+		case SetPrimaryClan: return LOCTEXT("SetPrimary", "Set Primary");
+
+		default: return FText::GetEmpty();
+		}
+	}
+
+	inline EClanActionLevel ToActionLevel(EClanActionType::Type State)
+	{
+		switch (State)
+		{
+		case AcceptClanRequest:
+		case SendClanRequest:
+			return EClanActionLevel::Emphasis;
+		case LeaveClan:
+			return EClanActionLevel::Critical;
+		default:
+			return EClanActionLevel::Action;
+		}
+	}
+};
+
 class FClanInfoViewModel
 	: public TSharedFromThis<FClanInfoViewModel>
 {
@@ -13,6 +67,10 @@ public:
 	virtual TSharedRef<class FFriendListViewModel > GetFriendListViewModel() = 0;
 	virtual FText GetListCountText() const = 0;
 	virtual FText GetClanBrushName() const = 0;
+	virtual void EnumerateActions(TArray<EClanActionType::Type>& OUTActions) = 0;
+
+	DECLARE_EVENT_OneParam(FClanInfoViewModel, FOpenClanDetails, const TSharedRef<class FClanInfoViewModel>&)
+	virtual FOpenClanDetails& OpenClanDetails() = 0;
 };
 
 /**
@@ -23,4 +81,6 @@ public:
 FACTORY(TSharedRef< FClanInfoViewModel >, FClanInfoViewModel,
 	const TSharedRef<class IClanInfo>& ClanInfo,
 	const TSharedRef<class IFriendListFactory>& FriendsListFactory,
-	const TSharedRef<class FFriendsAndChatManager>& FriendsAndChatManager);
\ No newline at end of file
+	const TSharedRef<class FFriendsAndChatManager>& FriendsAndChatManager);
+
+#undef LOCTEXT_NAMESPACE
\ No newline at end of file
diff --git a/Engine/Source/Developer/FriendsAndChat/Private/Layers/Presentation/ClanListViewModel.cpp b/Engine/Source/Developer/FriendsAndChat/Private/Layers/Presentation/ClanListViewModel.cpp
--- a/Engine/Source/Developer/FriendsAndChat/Private/Layers/Presentation/ClanListViewModel.cpp
+++ b/Engine/Source/Developer/FriendsAndChat/Private/Layers/Presentation/ClanListViewModel.cpp
@@ -1,7 +1,7 @@
 // Copyright 1998-2015 Epic Games, Inc. All Rights Reserved.
 
 #include "FriendsAndChatPrivatePCH.h"
-#include "ClanViewModel.h"
+#include "ClanCollectionViewModel.h"
 #include "ClanListViewModel.h"
 
 #define LOCTEXT_NAMESPACE "FriendsAndChat"
@@ -10,7 +10,7 @@ class FClanListViewModelImpl : public FClanListViewModel
 {
 public:
 
-	FClanListViewModelImpl(const TSharedRef<FClanViewModel>& InClanViewModel, EClanDisplayLists::Type InListType)
+	FClanListViewModelImpl(const TSharedRef<FClanCollectionViewModel>& InClanViewModel, EClanDisplayLists::Type InListType)
 		: ViewModel(InClanViewModel)
 		, ListType(InListType)
 	{
@@ -21,7 +21,7 @@ public:
 		Uninitialize();
 	}
 
-	virtual const TArray<TSharedRef<IClanInfo>>& GetClanList() const override
+	virtual const TArray<TSharedRef<FClanInfoViewModel>>& GetClanList() const override
 	{
 		return ClanList;
 	}
@@ -75,17 +75,17 @@ private:
 
 
 private:
-	const TSharedRef<FClanViewModel> ViewModel;
+	const TSharedRef<FClanCollectionViewModel> ViewModel;
 	const EClanDisplayLists::Type ListType;
 
 	/** Holds the list of clans. */
-	TArray<TSharedRef<IClanInfo>> ClanList;
+	TArray<TSharedRef<FClanInfoViewModel>> ClanList;
 	FClanListUpdated ClanListUpdatedEvent;
 
 	friend FClanListViewModelFactory;
 };
 
-TSharedRef< FClanListViewModel > FClanListViewModelFactory::Create(const TSharedRef<class FClanViewModel>& ClanViewModel, EClanDisplayLists::Type ListType)
+TSharedRef< FClanListViewModel > FClanListViewModelFactory::Create(const TSharedRef<class FClanCollectionViewModel>& ClanViewModel, EClanDisplayLists::Type ListType)
 {
 	TSharedRef< FClanListViewModelImpl > ViewModel(new FClanListViewModelImpl(ClanViewModel, ListType));
 	ViewModel->Initialize();
diff --git a/Engine/Source/Developer/FriendsAndChat/Private/Layers/Presentation/ClanListViewModel.h b/Engine/Source/Developer/FriendsAndChat/Private/Layers/Presentation/ClanListViewModel.h
--- a/Engine/Source/Developer/FriendsAndChat/Private/Layers/Presentation/ClanListViewModel.h
+++ b/Engine/Source/Developer/FriendsAndChat/Private/Layers/Presentation/ClanListViewModel.h
@@ -32,7 +32,7 @@ class FClanListViewModel
 {
 public:
 	virtual ~FClanListViewModel() {}
-	virtual const TArray<TSharedRef<class IClanInfo>>& GetClanList() const = 0;
+	virtual const TArray<TSharedRef<class FClanInfoViewModel>>& GetClanList() const = 0;
 	virtual FText GetListCountText() const = 0;
 	virtual int32 GetListCount() const = 0;
 	virtual const FText GetListName() const = 0;
@@ -48,4 +48,4 @@ public:
  *
  * @return the newly created FClanListViewModel implementation.
  */
-FACTORY(TSharedRef< FClanListViewModel >, FClanListViewModel,const TSharedRef<class FClanViewModel>& ClanViewModel, EClanDisplayLists::Type ListType);
\ No newline at end of file
+FACTORY(TSharedRef< FClanListViewModel >, FClanListViewModel,const TSharedRef<class FClanCollectionViewModel>& ClanViewModel, EClanDisplayLists::Type ListType);
\ No newline at end of file
diff --git a/Engine/Source/Developer/FriendsAndChat/Private/Layers/Presentation/ClanViewModel.cpp b/Engine/Source/Developer/FriendsAndChat/Private/Layers/Presentation/ClanViewModel.cpp
deleted file mode 100644
diff --git a/Engine/Source/Developer/FriendsAndChat/Private/Layers/Presentation/ClanViewModel.h b/Engine/Source/Developer/FriendsAndChat/Private/Layers/Presentation/ClanViewModel.h
deleted file mode 100644
diff --git a/Engine/Source/Developer/FriendsAndChat/Private/Layers/Presentation/FriendViewModel.cpp b/Engine/Source/Developer/FriendsAndChat/Private/Layers/Presentation/FriendViewModel.cpp
--- a/Engine/Source/Developer/FriendsAndChat/Private/Layers/Presentation/FriendViewModel.cpp
+++ b/Engine/Source/Developer/FriendsAndChat/Private/Layers/Presentation/FriendViewModel.cpp
@@ -118,7 +118,7 @@ public:
 			case EFriendActionType::RejectFriendRequest:
 			case EFriendActionType::CancelFriendRequest:
 			{
-				RemoveFriend(EFriendActionType::ToText(ActionType).ToString());
+				RemoveFriend(ActionType);
 				break;
 			}
 			case EFriendActionType::SendFriendRequest : 
@@ -257,9 +257,9 @@ public:
 
 private:
 
-	void RemoveFriend(const FString& Action) const
+	void RemoveFriend(EFriendActionType::Type Reason) const
 	{
-		FriendsAndChatManager.Pin()->DeleteFriend(FriendItem, Action);
+		FriendsAndChatManager.Pin()->DeleteFriend(FriendItem, Reason);
 	}
 
 	void AcceptFriend() const
diff --git a/Engine/Source/Developer/FriendsAndChat/Private/Layers/Presentation/FriendsViewModel.cpp b/Engine/Source/Developer/FriendsAndChat/Private/Layers/Presentation/FriendsViewModel.cpp
--- a/Engine/Source/Developer/FriendsAndChat/Private/Layers/Presentation/FriendsViewModel.cpp
+++ b/Engine/Source/Developer/FriendsAndChat/Private/Layers/Presentation/FriendsViewModel.cpp
@@ -6,7 +6,7 @@
 #include "FriendsStatusViewModel.h"
 #include "FriendListViewModel.h"
 #include "FriendsUserViewModel.h"
-#include "ClanViewModel.h"
+#include "ClanCollectionViewModel.h"
 #include "ClanRepository.h"
 #include "IFriendList.h"
 #include "FriendsNavigationService.h"
@@ -41,9 +41,9 @@ public:
 		return FFriendListViewModelFactory::Create(FriendsListFactory->Create(ListType), ListType);
 	}
 
-	virtual TSharedRef< FClanViewModel > GetClanViewModel() override
+	virtual TSharedRef< FClanCollectionViewModel > GetClanCollectionViewModel() override
 	{
-		return FClanViewModelFactory::Create(ClanRepository, FriendsListFactory, FriendsAndChatManager.Pin().ToSharedRef());
+		return FClanCollectionViewModelFactory::Create(ClanRepository, FriendsListFactory, FriendsAndChatManager.Pin().ToSharedRef());
 	}
 
 	virtual void RequestFriend(const FText& FriendName) const override
diff --git a/Engine/Source/Developer/FriendsAndChat/Private/Layers/Presentation/FriendsViewModel.h b/Engine/Source/Developer/FriendsAndChat/Private/Layers/Presentation/FriendsViewModel.h
--- a/Engine/Source/Developer/FriendsAndChat/Private/Layers/Presentation/FriendsViewModel.h
+++ b/Engine/Source/Developer/FriendsAndChat/Private/Layers/Presentation/FriendsViewModel.h
@@ -12,7 +12,7 @@ public:
 	virtual TSharedRef< class FFriendsUserViewModel > GetUserViewModel() = 0;
 	virtual TSharedRef< class FFriendsStatusViewModel > GetStatusViewModel() = 0;
 	virtual TSharedRef< class FFriendListViewModel > GetFriendListViewModel(EFriendsDisplayLists::Type ListType) = 0;
-	virtual TSharedRef< class FClanViewModel > GetClanViewModel() = 0;
+	virtual TSharedRef< class FClanCollectionViewModel > GetClanCollectionViewModel() = 0;
 	virtual void RequestFriend(const FText& FriendName) const = 0;
 	virtual EVisibility GetGlobalChatButtonVisibility() const = 0;
 	virtual void DisplayGlobalChatWindow() const = 0;
diff --git a/Engine/Source/Developer/FriendsAndChat/Private/Layers/UI/SChatMarkupTips.cpp b/Engine/Source/Developer/FriendsAndChat/Private/Layers/UI/SChatMarkupTips.cpp
--- a/Engine/Source/Developer/FriendsAndChat/Private/Layers/UI/SChatMarkupTips.cpp
+++ b/Engine/Source/Developer/FriendsAndChat/Private/Layers/UI/SChatMarkupTips.cpp
@@ -18,6 +18,7 @@ public:
 	{
 		ViewModel = InViewModel;
 		ViewModel->OnChatTipAvailable().AddSP(this, &SChatMarkupTipsImpl::HandleTipAvailable);
+		MarkupStyle = *InArgs._MarkupStyle;
 
 		SUserWidget::Construct(SUserWidget::FArguments()
 		[
@@ -39,9 +40,15 @@ private:
 				.AutoHeight()
 				[
 					SNew(SButton)
+					.ButtonStyle(&MarkupStyle.MarkupButtonStyle)
 					.ButtonColorAndOpacity(this, &SChatMarkupTipsImpl::GetActionBackgroundColor, TWeakPtr<IChatTip>(ChatTip))
-					.Text(ChatTip->GetTipText())
 					.OnClicked(this, &SChatMarkupTipsImpl::HandleActionClicked, ChatTip)
+					[
+						SNew(STextBlock)
+						.Text(ChatTip->GetTipText())
+						.TextStyle(&MarkupStyle.MarkupTextStyle)
+						.ColorAndOpacity(FLinearColor::White)
+					]
 				];
 			}
 		}
@@ -55,12 +62,9 @@ private:
 	FSlateColor GetActionBackgroundColor(TWeakPtr<IChatTip> TabPtr) const
 	{
 		TSharedPtr<IChatTip> Tab = TabPtr.Pin();
-		if (Tab.IsValid())
+		if (Tab.IsValid() && Tab == ViewModel->GetActiveTip())
 		{
-			if (Tab == ViewModel->GetActiveTip())
-			{
-				return FLinearColor::White;
-			}
+			return FLinearColor::Black;
 		}
 		return FLinearColor::Gray;
 	}
@@ -68,6 +72,8 @@ private:
 private:
 	TSharedPtr<FChatTipViewModel> ViewModel;
 	TSharedPtr<SVerticalBox> ChatTipBox;
+	/** Holds the style to use when making the widget. */
+	FFriendsMarkupStyle MarkupStyle;
 };
 
 TSharedRef<SChatMarkupTips> SChatMarkupTips::New()
diff --git a/Engine/Source/Developer/FriendsAndChat/Private/Layers/UI/SChatMarkupTips.h b/Engine/Source/Developer/FriendsAndChat/Private/Layers/UI/SChatMarkupTips.h
--- a/Engine/Source/Developer/FriendsAndChat/Private/Layers/UI/SChatMarkupTips.h
+++ b/Engine/Source/Developer/FriendsAndChat/Private/Layers/UI/SChatMarkupTips.h
@@ -7,6 +7,7 @@ class SChatMarkupTips : public SUserWidget
 public:
 	SLATE_USER_ARGS(SChatMarkupTips)
 	{ }
+		SLATE_ARGUMENT(const FFriendsMarkupStyle*, MarkupStyle)
 	SLATE_END_ARGS()
 
 	/**
diff --git a/Engine/Source/Developer/FriendsAndChat/Private/Layers/UI/SChatWindow.cpp b/Engine/Source/Developer/FriendsAndChat/Private/Layers/UI/SChatWindow.cpp
--- a/Engine/Source/Developer/FriendsAndChat/Private/Layers/UI/SChatWindow.cpp
+++ b/Engine/Source/Developer/FriendsAndChat/Private/Layers/UI/SChatWindow.cpp
@@ -11,6 +11,7 @@
 #include "SFriendActions.h"
 #include "SChatEntryWidget.h"
 #include "RichTextLayoutMarshaller.h"
+#include "TextDecorators.h"
 
 #define LOCTEXT_NAMESPACE "SChatWindow"
 
@@ -37,10 +38,12 @@ public:
 
 		OnHyperlinkClicked = FSlateHyperlinkRun::FOnClick::CreateSP(SharedThis(this), &SChatWindowImpl::HandleNameClicked);
 		OnChannelHyperlinkClicked = FSlateHyperlinkRun::FOnClick::CreateSP(SharedThis(this), &SChatWindowImpl::HandleChannelClicked);
+		OnTimeStampCreated = FWidgetDecorator::FCreateWidget::CreateSP(SharedThis(this), &SChatWindowImpl::HandleTimeStampDisplay);
 
 		RichTextMarshaller->AppendInlineDecorator(FHyperlinkDecorator::Create(TEXT("UserName"), OnHyperlinkClicked));
 		RichTextMarshaller->AppendInlineDecorator(FHyperlinkDecorator::Create(TEXT("Channel"), OnChannelHyperlinkClicked));
 		RichTextMarshaller->AppendInlineDecorator(FImageDecorator::Create(TEXT("img"), &FFriendsAndChatModuleStyle::Get()));
+		RichTextMarshaller->AppendInlineDecorator(FWidgetDecorator::Create(TEXT("TimeStamp"), OnTimeStampCreated));
 
 		TimeTransparency = 0.0f;
 
@@ -97,6 +100,7 @@ public:
 						.Visibility(ViewModelPtr, &FChatViewModel::GetTipVisibility)
 						[
 							SNew(SChatMarkupTips, ViewModel->GetChatTipViewModel())
+							.MarkupStyle(&FriendStyle.FriendsMarkupStyle)
 						]
 					]
 				]
@@ -118,16 +122,8 @@ public:
 							[
 								SNew(SVerticalBox)
 								+ SVerticalBox::Slot()
+								.Expose(ChannelTextSlot)
 								.AutoHeight()
-								[
-									SNew(SHorizontalBox)
-									+SHorizontalBox::Slot()
-									.AutoWidth()
-									[
-										SNew(STextBlock)
-										.Text(ViewModelPtr, &FChatViewModel::GetOutgoingChannelText)
-									]
-								]
 								+ SVerticalBox::Slot()
 								[
 									SAssignNew(ChatTextBox, SChatEntryWidget, ViewModel.ToSharedRef())
@@ -159,6 +155,14 @@ public:
 			]
 		]);
 
+		if(ChannelTextSlot && ViewModel->MultiChat())
+		{
+			ChannelTextSlot->AttachWidget(
+				SNew(STextBlock)
+				.Text(ViewModelPtr, &FChatViewModel::GetOutgoingChannelText)
+			);
+		}
+
 		RefreshChatList();
 	}
 
@@ -197,11 +201,6 @@ public:
 				{
 					FadeAnimation.PlayReverse(this->AsShared());
 				}
-
-				if (FadeAnimation.IsPlaying())
-				{
-					ChatScrollBox->ScrollToEnd();
-				}
 			}
 		}
 
@@ -297,7 +296,10 @@ private:
 		if (RichText.IsValid() && ChatScrollBox.IsValid())
 		{
 			RichText->Refresh();
-			ChatScrollBox->ScrollToEnd();
+			if(ExternalScrollbar.IsValid() && ExternalScrollbar->DistanceFromBottom() == 0)
+			{
+				ChatScrollBox->ScrollToEnd();
+			}
 		}
 	}
 
@@ -338,6 +340,113 @@ private:
 		}
 	}
 
+	FSlateWidgetRun::FWidgetRunInfo HandleTimeStampDisplay( const FTextRunInfo& RunInfo, const ISlateStyle* Style ) const
+	{
+		// Internal class for timestamps
+		class SChatTimeStamp : public SCompoundWidget
+		{
+			SLATE_BEGIN_ARGS(SChatTimeStamp)
+				: _DateTime(FDateTime::Now())
+			{ }
+				SLATE_ARGUMENT(FDateTime, DateTime)
+			SLATE_END_ARGS()
+
+			virtual ~SChatTimeStamp()
+			{
+				FTicker::GetCoreTicker().RemoveTicker(TickerHandle);
+			}
+
+			void Construct(const FArguments& InArgs)
+			{
+				TimeStamp = InArgs._DateTime;
+				ChildSlot
+				[
+					SNew(STextBlock)
+					.Text(this, &SChatTimeStamp::GetMessageTimeText)
+					.Font(FFriendsAndChatModuleStyle::GetStyleService()->GetNormalFont())
+				];
+
+				if (UpdateTimeStamp(60)) //this will broadcast a changed event
+				{
+					TickerHandle = FTicker::GetCoreTicker().AddTicker(FTickerDelegate::CreateSP(this, &SChatTimeStamp::UpdateTimeStamp), 60);
+				}
+			}
+
+			private:
+
+			bool UpdateTimeStamp(float Delay)
+			{
+				bool ContinueTicking = true;
+
+				FDateTime Now = FDateTime::Now();
+
+				if ((Now - TimeStamp).GetDuration() <= FTimespan::FromMinutes(1.0))
+				{
+					static FText NowTimeStamp = NSLOCTEXT("SChatWindow", "Now_TimeStamp", "Now");
+					MessageTimeAsText = NowTimeStamp;
+				}
+				else if ((Now - TimeStamp).GetDuration() <= FTimespan::FromMinutes(2.0))
+				{
+					static FText OneMinuteTimeStamp = NSLOCTEXT("SChatWindow", "1_Minute_TimeStamp", "1 min");
+					MessageTimeAsText = OneMinuteTimeStamp;
+				}
+				else if ((Now - TimeStamp).GetDuration() <= FTimespan::FromMinutes(3.0))
+				{
+					static FText TwoMinuteTimeStamp = NSLOCTEXT("SChatWindow", "2_Minute_TimeStamp", "2 min");
+					MessageTimeAsText = TwoMinuteTimeStamp;
+				}
+				else if ((Now - TimeStamp).GetDuration() <= FTimespan::FromMinutes(4.0))
+				{
+					static FText ThreeMinuteTimeStamp = NSLOCTEXT("SChatWindow", "3_Minute_TimeStamp", "3 min");
+					MessageTimeAsText = ThreeMinuteTimeStamp;
+				}
+				else if ((Now - TimeStamp).GetDuration() <= FTimespan::FromMinutes(5.0))
+				{
+					static FText FourMinuteTimeStamp = NSLOCTEXT("SChatWindow", "4_Minute_TimeStamp", "4 min");
+					MessageTimeAsText = FourMinuteTimeStamp;
+				}
+				else if ((Now - TimeStamp).GetDuration() <= FTimespan::FromMinutes(6.0))
+				{
+					static FText FiveMinuteTimeStamp = NSLOCTEXT("SChatWindow", "5_Minute_TimeStamp", "5 min");
+					MessageTimeAsText = FiveMinuteTimeStamp;
+				}
+				else
+				{
+					MessageTimeAsText = FText::AsTime(TimeStamp, EDateTimeStyle::Short, "GMT");
+					ContinueTicking = false;
+				}
+
+				return ContinueTicking;
+			}
+
+			FText GetMessageTimeText() const
+			{
+				return MessageTimeAsText;
+			}
+
+			private:
+				FText MessageTimeAsText;
+				FDelegateHandle TickerHandle;
+				FDateTime TimeStamp;
+		};
+
+		TSharedRef< SWidget > Widget = SNullWidget::NullWidget;
+
+		FDateTime TimeStamp;
+		FString TimeStampSting = RunInfo.Content.ToString();
+		if (!TimeStampSting.IsEmpty() && FDateTime::Parse(TimeStampSting, TimeStamp))
+		{
+			Widget = SNew( SChatTimeStamp )
+			.DateTime(TimeStamp);
+		}
+
+
+		TSharedRef< FSlateFontMeasure > FontMeasure = FSlateApplication::Get().GetRenderer()->GetFontMeasureService();
+		int16 Baseline = FontMeasure->GetBaseline(FFriendsAndChatModuleStyle::GetStyleService()->GetNormalFont());
+
+		return FSlateWidgetRun::FWidgetRunInfo(Widget, Baseline);
+	}
+
 	FText FormatListMessageText(TSharedRef<FChatItemViewModel> ChatItem, bool bGrouped)
 	{
 		const FText& TextToSet = ChatItem->GetMessage();
@@ -373,7 +482,10 @@ private:
 			Args.Add(TEXT("SenderID"), FText::FromString(ChatItem->GetSenderID()->ToString()));
 		}
 
-		Args.Add(TEXT("MessageTime"), ChatItem->GetSimpleMessageTimeText());
+		FString MessageTime = ChatItem->GetMessageTime().ToString();
+		FString TimeStampMarkup = "<TimeStamp>" + MessageTime + "</>";
+		Args.Add(TEXT("TimeStampMarkup"), FText::FromString(TimeStampMarkup));
+
 		Args.Add(TEXT("NameStyle"), FText::FromString(HyperlinkStyle));
 		Args.Add(TEXT("Message"), TextToSet);
 
@@ -381,7 +493,7 @@ private:
 
 		if (bGrouped)
 		{
-			ListMessage = FText::Format(NSLOCTEXT("FChatItemViewModel", "DisplayNameAndMessage", "{Message}"), TextToSet);
+			ListMessage = FText::Format(NSLOCTEXT("SChatWindow", "DisplayMessage", "{Message}"), TextToSet);
 		}
 		else
 		{
@@ -389,11 +501,11 @@ private:
 			{
 				if (ChatItem->IsFromSelf())
 				{
-					ListMessage = FText::Format(NSLOCTEXT("FChatItemViewModel", "DisplayNameAndMessage", "{MessageTime} {SenderName} : {Message}"), Args);
+					ListMessage = FText::Format(NSLOCTEXT("FChatItemViewModel", "DisplayNameAndMessage", "{TimeStampMarkup} {SenderName} : {Message}"), Args);
 				}
 				else
 				{
-					ListMessage = FText::Format(NSLOCTEXT("FChatItemViewModel", "DisplayNameAndMessage", "{MessageTime} <a id=\"UserName\" uid=\"{SenderID}\" Username=\"{SenderName}\" style=\"{NameStyle}\">{SenderName}</> : {Message}"), Args);
+					ListMessage = FText::Format(NSLOCTEXT("SChatWindow", "DisplayNameAndMessageWithMarkup", "{TimeStampMarkup} <a id=\"UserName\" uid=\"{SenderID}\" Username=\"{SenderName}\" style=\"{NameStyle}\">{SenderName}</> : {Message}"), Args);
 				}
 			}
 			else
@@ -402,22 +514,22 @@ private:
 				{
 					if (ChatItem->IsFromSelf())
 					{
-						ListMessage = FText::Format(NSLOCTEXT("FChatItemViewModel", "DisplayNameAndMessage", "{MessageTime} <a id=\"Channel\" Command=\"{Channel}\" style=\"{NameStyle}\">{RoomName}</> to <a id=\"UserName\" uid=\"{RecipientID}\" Username=\"{RecipientName}\" style=\"{NameStyle}\">{RecipientName}</> : {Message}"), Args);
+						ListMessage = FText::Format(NSLOCTEXT("SChatWindow", "DisplayNameAndMessageWhisperFromSelf", "{TimeStampMarkup} <a id=\"Channel\" Command=\"{Channel}\" style=\"{NameStyle}\">{RoomName}</> to <a id=\"UserName\" uid=\"{RecipientID}\" Username=\"{RecipientName}\" style=\"{NameStyle}\">{RecipientName}</> : {Message}"), Args);
 					}
 					else
 					{
-						ListMessage = FText::Format(NSLOCTEXT("FChatItemViewModel", "DisplayNameAndMessage", "{MessageTime} <a id=\"Channel\" Command=\"{Channel}\" style=\"{NameStyle}\">{RoomName}</> from <a id=\"UserName\" uid=\"{SenderID}\" Username=\"{SenderName}\" style=\"{NameStyle}\">{SenderName}</> : {Message}"), Args);
+						ListMessage = FText::Format(NSLOCTEXT("SChatWindow", "DisplayNameAndMessageWhisper", "{TimeStampMarkup} <a id=\"Channel\" Command=\"{Channel}\" style=\"{NameStyle}\">{RoomName}</> from <a id=\"UserName\" uid=\"{SenderID}\" Username=\"{SenderName}\" style=\"{NameStyle}\">{SenderName}</> : {Message}"), Args);
 					}
 				}
 				else
 				{
 					if (ChatItem->IsFromSelf())
 					{
-						ListMessage = FText::Format(NSLOCTEXT("FChatItemViewModel", "DisplayNameAndMessage", "{MessageTime} <a id=\"Channel\" Command=\"{Channel}\" style=\"{NameStyle}\">{RoomName}</> {SenderName} : {Message}"), Args);
+						ListMessage = FText::Format(NSLOCTEXT("SChatWindow", "DisplayNameAndMessageOtherFromSelf", "{TimeStampMarkup} <a id=\"Channel\" Command=\"{Channel}\" style=\"{NameStyle}\">{RoomName}</> {SenderName} : {Message}"), Args);
 					}
 					else
 					{
-						ListMessage = FText::Format(NSLOCTEXT("FChatItemViewModel", "DisplayNameAndMessage", "{MessageTime} <a id=\"Channel\" Command=\"{Channel}\" style=\"{NameStyle}\">{RoomName}</> <a id=\"UserName\" uid=\"{SenderID}\" Username=\"{SenderName}\" style=\"{NameStyle}\">{SenderName}</> : {Message}"), Args);
+						ListMessage = FText::Format(NSLOCTEXT("SChatWindow", "DisplayNameAndMessageOther", "{TimeStampMarkup} <a id=\"Channel\" Command=\"{Channel}\" style=\"{NameStyle}\">{RoomName}</> <a id=\"UserName\" uid=\"{SenderID}\" Username=\"{SenderName}\" style=\"{NameStyle}\">{SenderName}</> : {Message}"), Args);
 					}
 				}
 			}
@@ -473,6 +585,8 @@ private:
 
 	TSharedPtr<SMultiLineEditableTextBox> RichText;
 
+	SVerticalBox::FSlot* ChannelTextSlot;
+
 	FTextBuilder ChatText;
 
 	int32 LastMessageIndex;
@@ -516,6 +630,7 @@ private:
 
 	FSlateHyperlinkRun::FOnClick OnHyperlinkClicked;
 	FSlateHyperlinkRun::FOnClick OnChannelHyperlinkClicked;
+	FWidgetDecorator::FCreateWidget OnTimeStampCreated;
 
 	static const float CHAT_HINT_UPDATE_THROTTLE;
 };
diff --git a/Engine/Source/Developer/FriendsAndChat/Private/Layers/UI/SClanContainer.cpp b/Engine/Source/Developer/FriendsAndChat/Private/Layers/UI/SClanContainer.cpp
--- a/Engine/Source/Developer/FriendsAndChat/Private/Layers/UI/SClanContainer.cpp
+++ b/Engine/Source/Developer/FriendsAndChat/Private/Layers/UI/SClanContainer.cpp
@@ -5,7 +5,7 @@
 #include "SClanDetails.h"
 #include "SClanHome.h"
 #include "SWidgetSwitcher.h"
-#include "ClanViewModel.h"
+#include "ClanCollectionViewModel.h"
 
 #define LOCTEXT_NAMESPACE ""
 
@@ -16,41 +16,56 @@ class SClanContainerImpl : public SClanContainer
 {
 public:
 
-	void Construct(const FArguments& InArgs, const TSharedRef<FClanViewModel>& InViewModel)
+	void Construct(const FArguments& InArgs, const TSharedRef<FClanCollectionViewModel>& InViewModel)
 	{
 		FriendStyle = *InArgs._FriendStyle;
 		ViewModel = InViewModel;
 		MenuMethod = InArgs._Method;
 
+		ViewModel->OpenClanDetails().AddSP(this, &SClanContainerImpl::OnOpenClanDetails);
+
+		ExternalScrollbar = SNew(SScrollBar)
+			.Thickness(FVector2D(4, 4))
+			.Style(&FriendStyle.ScrollBarStyle)
+			.AlwaysShowScrollbar(true);
+
 		SUserWidget::Construct(SUserWidget::FArguments()
 		[
-			SNew(SVerticalBox)
-			+ SVerticalBox::Slot()
-			.VAlign(VAlign_Top)
-			.HAlign(HAlign_Fill)
+			SNew(SBorder)
+			.BorderImage(&FriendStyle.FriendsContainerBackground)
+			.Padding(0)
 			[
-				SAssignNew(Switcher, SWidgetSwitcher)
-				+SWidgetSwitcher::Slot()
+				SNew(SOverlay)
+				+ SOverlay::Slot()
 				[
-					SNew(SVerticalBox)
-					+SVerticalBox::Slot()
-					.AutoHeight()
-					[
-						SNew(SButton)
-						.Text(FText::FromString("Details"))
-						.OnClicked(this, &SClanContainerImpl::HandleWidgetSwitcherClicked, 1)
-					]
-					+SVerticalBox::Slot()
+					SNew(SScrollBox)
+					.ExternalScrollbar(ExternalScrollbar.ToSharedRef())
+					+ SScrollBox::Slot()
+					.HAlign(HAlign_Fill)
 					[
-						SNew(SClanHome, ViewModel.ToSharedRef())
-						.FriendStyle(&FriendStyle)
+						SNew(SVerticalBox)
+						+ SVerticalBox::Slot()
+						.VAlign(VAlign_Top)
+						.AutoHeight()
+						[
+							SAssignNew(Switcher, SWidgetSwitcher)
+							+ SWidgetSwitcher::Slot()
+							[
+								SNew(SClanHome, ViewModel.ToSharedRef())
+								.FriendStyle(&FriendStyle)
+							]
+							+ SWidgetSwitcher::Slot()
+							[
+								SAssignNew(ClanDetailsBox, SBox)
+							]
+						]
 					]
 				]
-				+SWidgetSwitcher::Slot()
+				+ SOverlay::Slot()
+				.HAlign(HAlign_Right)
+				.Padding(FMargin(10, 20, 10, 20))
 				[
-					SNew(SClanDetails, ViewModel->GetClanInfoViewModel().ToSharedRef())
-					.FriendStyle(&FriendStyle)
-					.OnHomeClicked(this, &SClanContainerImpl::HandleWidgetSwitcherClicked, 0)
+					ExternalScrollbar.ToSharedRef()
 				]
 			]
 		]);
@@ -58,6 +73,15 @@ public:
 
 private:
 
+	void OnOpenClanDetails(const TSharedRef < FClanInfoViewModel > &ClanInfoViewModel)
+	{
+		TSharedRef<SClanDetails> ClanDetails = SNew(SClanDetails, ClanInfoViewModel)
+			.FriendStyle(&FriendStyle)
+			.OnHomeClicked(this, &SClanContainerImpl::HandleWidgetSwitcherClicked, 0);
+		ClanDetailsBox->SetContent(ClanDetails);
+		Switcher->SetActiveWidgetIndex(1);
+	}
+
 	FReply HandleWidgetSwitcherClicked(int32 Index)
 	{
 		Switcher->SetActiveWidgetIndex(Index);
@@ -67,11 +91,14 @@ private:
 	/** Holds the style to use when making the widget. */
 	FFriendsListStyle FriendStyle;
 
-	TSharedPtr<FClanViewModel> ViewModel;
+	TSharedPtr<SBox> ClanDetailsBox;
+	TSharedPtr<FClanCollectionViewModel> ViewModel;
 
 	// Holds the list container switcher
 	TSharedPtr<SWidgetSwitcher> Switcher;
 
+	TSharedPtr<SScrollBar> ExternalScrollbar;
+
 	EPopupMethod MenuMethod;
 };
 
diff --git a/Engine/Source/Developer/FriendsAndChat/Private/Layers/UI/SClanContainer.h b/Engine/Source/Developer/FriendsAndChat/Private/Layers/UI/SClanContainer.h
--- a/Engine/Source/Developer/FriendsAndChat/Private/Layers/UI/SClanContainer.h
+++ b/Engine/Source/Developer/FriendsAndChat/Private/Layers/UI/SClanContainer.h
@@ -11,5 +11,5 @@ public:
 	SLATE_ARGUMENT(EPopupMethod, Method)
 	SLATE_END_ARGS()
 
-	virtual void Construct(const FArguments& InArgs, const TSharedRef<class FClanViewModel>& ViewModel) = 0;
+	virtual void Construct(const FArguments& InArgs, const TSharedRef<class FClanCollectionViewModel>& ViewModel) = 0;
 };
diff --git a/Engine/Source/Developer/FriendsAndChat/Private/Layers/UI/SClanHome.cpp b/Engine/Source/Developer/FriendsAndChat/Private/Layers/UI/SClanHome.cpp
--- a/Engine/Source/Developer/FriendsAndChat/Private/Layers/UI/SClanHome.cpp
+++ b/Engine/Source/Developer/FriendsAndChat/Private/Layers/UI/SClanHome.cpp
@@ -4,7 +4,7 @@
 #include "SClanHome.h"
 #include "SClanList.h"
 #include "SClanListContainer.h"
-#include "ClanViewModel.h"
+#include "ClanCollectionViewModel.h"
 #include "ClanListViewModel.h"
 
 #define LOCTEXT_NAMESPACE ""
@@ -16,7 +16,7 @@ class SClanHomeImpl : public SClanHome
 {
 public:
 
-	void Construct(const FArguments& InArgs, const TSharedRef<FClanViewModel>& InViewModel)
+	void Construct(const FArguments& InArgs, const TSharedRef<FClanCollectionViewModel>& InViewModel)
 	{
 		FriendStyle = *InArgs._FriendStyle;
 		ViewModel = InViewModel;
@@ -25,19 +25,35 @@ public:
 		[
 			SNew(SVerticalBox)
 			+ SVerticalBox::Slot()
+			.AutoHeight()
 			.Padding(10)
+			[
+				SNew(SButton)
+				.Visibility(this, &SClanHomeImpl::ViewClansVisibility)
+				.ButtonStyle(&FriendStyle.GlobalChatButtonStyle)
+				[
+					SNew(STextBlock)
+					.Font(FriendStyle.FriendsNormalFontStyle.FriendsFontNormalBold)
+					.ColorAndOpacity(FriendStyle.FriendsNormalFontStyle.DefaultFontColor)
+					.Text(LOCTEXT("ViewClans","View Clans"))
+				]
+			]
+			+ SVerticalBox::Slot()
+			.AutoHeight()
+			.Padding(FMargin(10, 10, 30, 0))
 			[
 				SNew(SClanListContainer, InViewModel)
 				.FriendStyle(&FriendStyle)
-				.Visibility(EVisibility::Visible)
+				.Visibility(this, &SClanHomeImpl::ClanListVisibility)
 				.ClanListType(EClanDisplayLists::DefaultDisplay)
 			]
 			+ SVerticalBox::Slot()
-			.Padding(10)
+			.AutoHeight()
+			.Padding(FMargin(10, 10, 30, 0))
 			[
 				SNew(SClanListContainer, InViewModel)
 				.FriendStyle(&FriendStyle)
-				.Visibility(EVisibility::Collapsed)
+				.Visibility(this, &SClanHomeImpl::InviteListVisibility)
 				.ClanListType(EClanDisplayLists::ClanRequestsDisplay)
 			]
 		]);
@@ -45,10 +61,37 @@ public:
 
 private:
 
+	EVisibility ViewClansVisibility() const
+	{
+		if (ViewModel.IsValid() && ViewModel->NumberOfJoinedClans() == 0)
+		{
+			return EVisibility::Visible;
+		}
+		return EVisibility::Collapsed;
+	}
+
+	EVisibility ClanListVisibility() const
+	{
+		if (ViewModel.IsValid() && ViewModel->NumberOfJoinedClans() > 0)
+		{
+			return EVisibility::Visible;
+		}
+		return EVisibility::Collapsed;
+	}
+
+	EVisibility InviteListVisibility() const
+	{
+		if (ViewModel.IsValid() && ViewModel->NumberOfJoinedClans() > 0)
+		{
+			return EVisibility::Visible;
+		}
+		return EVisibility::Collapsed;
+	}
+
 	/** Holds the style to use when making the widget. */
 	FFriendsListStyle FriendStyle;
 
-	TSharedPtr<FClanViewModel> ViewModel;
+	TSharedPtr<FClanCollectionViewModel> ViewModel;
 };
 
 TSharedRef<SClanHome> SClanHome::New()
diff --git a/Engine/Source/Developer/FriendsAndChat/Private/Layers/UI/SClanHome.h b/Engine/Source/Developer/FriendsAndChat/Private/Layers/UI/SClanHome.h
--- a/Engine/Source/Developer/FriendsAndChat/Private/Layers/UI/SClanHome.h
+++ b/Engine/Source/Developer/FriendsAndChat/Private/Layers/UI/SClanHome.h
@@ -10,5 +10,5 @@ public:
 	SLATE_ARGUMENT(const FFriendsListStyle*, FriendStyle)
 	SLATE_END_ARGS()
 
-	virtual void Construct(const FArguments& InArgs, const TSharedRef<class FClanViewModel>& ViewModel) = 0;
+	virtual void Construct(const FArguments& InArgs, const TSharedRef<class FClanCollectionViewModel>& ViewModel) = 0;
 };
diff --git a/Engine/Source/Developer/FriendsAndChat/Private/Layers/UI/SClanItem.cpp b/Engine/Source/Developer/FriendsAndChat/Private/Layers/UI/SClanItem.cpp
--- a/Engine/Source/Developer/FriendsAndChat/Private/Layers/UI/SClanItem.cpp
+++ b/Engine/Source/Developer/FriendsAndChat/Private/Layers/UI/SClanItem.cpp
@@ -2,7 +2,7 @@
 
 #include "FriendsAndChatPrivatePCH.h"
 #include "SClanItem.h"
-#include "ClanInfo.h"
+#include "ClanInfoViewModel.h"
 #include "SFriendsAndChatCombo.h"
 
 #define LOCTEXT_NAMESPACE "SClassItem"
@@ -11,10 +11,10 @@ class SClanItemImpl : public SClanItem
 {
 public:
 
-	void Construct(const FArguments& InArgs, const TSharedRef<class IClanInfo>& InClanInfo)
+	void Construct(const FArguments& InArgs, const TSharedRef<FClanInfoViewModel>& InClanInfoViewModel)
 	{
 		FriendStyle = *InArgs._FriendStyle;
-		ClanInfo = InClanInfo;
+		ClanInfoViewModel = InClanInfoViewModel;
 
 		FFriendsAndChatComboButtonStyle ActionButtonStyle;
 		ActionButtonStyle.ComboButtonStyle = &FriendStyle.FriendsComboStyle.ActionComboButtonStyle;
@@ -79,6 +79,7 @@ public:
 					.Padding(0, 10)
 					[
 						SNew(SButton)
+						.OnClicked(this, &SClanItemImpl::OnOpenClanDetails)
 						.ButtonStyle(&FriendStyle.GlobalChatButtonStyle)
 						[
 							SNew(SHorizontalBox)
@@ -114,7 +115,7 @@ public:
 					SNew(STextBlock)
 					.Font(FriendStyle.FriendsNormalFontStyle.FriendsFontNormalBold)
 					.ColorAndOpacity(FriendStyle.FriendsNormalFontStyle.DefaultFontColor)
-					.Text(ClanInfo->GetTitle())
+					.Text(ClanInfoViewModel->GetClanTitle())
 				]
 			]
 		]);
@@ -122,10 +123,25 @@ public:
 
 private:
 
+	FReply OnOpenClanDetails()
+	{
+		ClanInfoViewModel->OpenClanDetails().Broadcast(ClanInfoViewModel.ToSharedRef());
+		return FReply::Handled();
+	}
+
 	SFriendsAndChatCombo::FItemsArray GetActionItems() const
 	{
 		SFriendsAndChatCombo::FItemsArray ActionItems;
 		// Need to Add some actions
+
+		TArray<EClanActionType::Type> Actions;
+		ClanInfoViewModel->EnumerateActions(Actions);
+ 
+ 		for (const auto& ClanAction : Actions)
+ 		{
+			ActionItems.AddItem(EClanActionType::ToText(ClanAction), nullptr, FName(*EClanActionType::ToText(ClanAction).ToString()), true);
+ 		}
+
 		return ActionItems;
 	}
 
@@ -149,11 +165,11 @@ private:
 
 	FText GetClanMemberCountText() const
 	{
-		return FText::AsNumber(ClanInfo->GetClanMembersCount());
+		return FText::AsNumber(ClanInfoViewModel->GetMemberCount());
 	}
 
 
-	TSharedPtr<IClanInfo> ClanInfo;
+	TSharedPtr<FClanInfoViewModel> ClanInfoViewModel;
 
 	/** Holds the style to use when making the widget. */
 	FFriendsListStyle FriendStyle;
diff --git a/Engine/Source/Developer/FriendsAndChat/Private/Layers/UI/SClanItem.h b/Engine/Source/Developer/FriendsAndChat/Private/Layers/UI/SClanItem.h
--- a/Engine/Source/Developer/FriendsAndChat/Private/Layers/UI/SClanItem.h
+++ b/Engine/Source/Developer/FriendsAndChat/Private/Layers/UI/SClanItem.h
@@ -16,5 +16,5 @@ public:
 	 *
 	 * @param InArgs The Slate argument list.
 	 */
-	virtual void Construct( const FArguments& InArgs, const TSharedRef<class IClanInfo>& ClanInfo ) = 0;
+	 virtual void Construct(const FArguments& InArgs, const TSharedRef<class FClanInfoViewModel>& ClanInfoViewModel) = 0;
 };
diff --git a/Engine/Source/Developer/FriendsAndChat/Private/Layers/UI/SClanListContainer.cpp b/Engine/Source/Developer/FriendsAndChat/Private/Layers/UI/SClanListContainer.cpp
--- a/Engine/Source/Developer/FriendsAndChat/Private/Layers/UI/SClanListContainer.cpp
+++ b/Engine/Source/Developer/FriendsAndChat/Private/Layers/UI/SClanListContainer.cpp
@@ -14,7 +14,7 @@ class SClanListContainerImpl : public SClanListContainer
 {
 public:
 
-	void Construct(const FArguments& InArgs, const TSharedRef<FClanViewModel>& InViewModel)
+	void Construct(const FArguments& InArgs, const TSharedRef<FClanCollectionViewModel>& InViewModel)
 	{
 		FriendStyle = *InArgs._FriendStyle;
 		ClanListType = InArgs._ClanListType;
@@ -115,7 +115,7 @@ private:
 
 	EVisibility  ClansVisibility;
 
-	TSharedPtr<FClanViewModel> ViewModel;
+	TSharedPtr<FClanCollectionViewModel> ViewModel;
 
 	EPopupMethod MenuMethod;
 
diff --git a/Engine/Source/Developer/FriendsAndChat/Private/Layers/UI/SClanListContainer.h b/Engine/Source/Developer/FriendsAndChat/Private/Layers/UI/SClanListContainer.h
--- a/Engine/Source/Developer/FriendsAndChat/Private/Layers/UI/SClanListContainer.h
+++ b/Engine/Source/Developer/FriendsAndChat/Private/Layers/UI/SClanListContainer.h
@@ -17,5 +17,5 @@ public:
 	SLATE_ARGUMENT(EClanDisplayLists::Type, ClanListType)
 	SLATE_END_ARGS()
 
-	virtual void Construct(const FArguments& InArgs, const TSharedRef<class FClanViewModel>& ViewModel) = 0;
+	virtual void Construct(const FArguments& InArgs, const TSharedRef<class FClanCollectionViewModel>& ViewModel) = 0;
 };
diff --git a/Engine/Source/Developer/FriendsAndChat/Private/Layers/UI/SFriendsContainer.cpp b/Engine/Source/Developer/FriendsAndChat/Private/Layers/UI/SFriendsContainer.cpp
--- a/Engine/Source/Developer/FriendsAndChat/Private/Layers/UI/SFriendsContainer.cpp
+++ b/Engine/Source/Developer/FriendsAndChat/Private/Layers/UI/SFriendsContainer.cpp
@@ -3,7 +3,7 @@
 #include "FriendsAndChatPrivatePCH.h"
 #include "FriendsViewModel.h"
 #include "FriendListViewModel.h"
-#include "ClanViewModel.h"
+#include "ClanCollectionViewModel.h"
 #include "SFriendUserHeader.h"
 #include "SFriendsList.h"
 #include "SFriendsListContainer.h"
@@ -31,6 +31,7 @@ public:
 		FriendStyle = *InArgs._FriendStyle;
 		FriendListStyle = FriendStyle.FriendsListStyle;
 		FriendFontStyle =  *InArgs._FontStyle;
+		MaxFriendsNameLength = InArgs._MaxFriendNameLength;
 
 		for (int32 ListIndex = 0; ListIndex < EFriendsDisplayLists::MAX_None; ListIndex++)
 		{
@@ -197,7 +198,16 @@ private:
 			{
 				ViewModel->PerformAction();
 			}
-		}		
+		}
+	}
+
+	void HandleChatTextChanged(const FText& CurrentText)
+	{
+		FString CurrentTextString = CurrentText.ToString();
+		if(CurrentTextString.Len() > MaxFriendsNameLength)
+		{
+			FriendNameTextBox->SetText(FText::FromString(CurrentTextString.Left(MaxFriendsNameLength)));
+		}
 	}
 
 	EVisibility AddFriendBoxVisibility() const
@@ -251,7 +261,7 @@ private:
 
 	TSharedRef<SWidget> GetClanWidget()
 	{
-		return SNew(SClanContainer, ViewModel->GetClanViewModel())
+		return SNew(SClanContainer, ViewModel->GetClanCollectionViewModel())
 		.FriendStyle(&FriendListStyle);
 	}
 
@@ -378,6 +388,7 @@ private:
 								.HintText(LOCTEXT("AddFriendHint", "Add friend by account name or email"))
 								.Style(&FriendListStyle.AddFriendEditableTextStyle)
 								.OnTextCommitted(this, &SFriendsContainerImpl::HandleFriendEntered)
+								.OnTextChanged(this, &SFriendsContainerImpl::HandleChatTextChanged)
 							]
 						]
 					]
@@ -484,6 +495,9 @@ FSlateColor GetFriendTabTextColor(int32 WidgetIndex) const
 
 private:
 
+	// Max Friend Name Length
+	int32 MaxFriendsNameLength;
+
 	// Holds the Friends List view model
 	TSharedPtr<FFriendsViewModel> ViewModel;
 
diff --git a/Engine/Source/Developer/FriendsAndChat/Private/Layers/UI/SFriendsContainer.h b/Engine/Source/Developer/FriendsAndChat/Private/Layers/UI/SFriendsContainer.h
--- a/Engine/Source/Developer/FriendsAndChat/Private/Layers/UI/SFriendsContainer.h
+++ b/Engine/Source/Developer/FriendsAndChat/Private/Layers/UI/SFriendsContainer.h
@@ -10,10 +10,12 @@ public:
 
 	SLATE_USER_ARGS(SFriendsContainer)
 		: _Method(EPopupMethod::UseCurrentWindow)
+		, _MaxFriendNameLength(254)
 	 { }
 	SLATE_ARGUMENT(const FFriendsAndChatStyle*, FriendStyle)
 	SLATE_ARGUMENT(const FFriendsFontStyle*, FontStyle)
 	SLATE_ARGUMENT(EPopupMethod, Method)
+	SLATE_ARGUMENT(int32, MaxFriendNameLength)
 	SLATE_END_ARGS()
 
 	virtual void Construct(const FArguments& InArgs, const TSharedRef<FFriendsViewModel>& ViewModel) = 0;
diff --git a/Engine/Source/Developer/FriendsAndChat/Public/FriendsAndChatStyle.h b/Engine/Source/Developer/FriendsAndChat/Public/FriendsAndChatStyle.h
--- a/Engine/Source/Developer/FriendsAndChat/Public/FriendsAndChatStyle.h
+++ b/Engine/Source/Developer/FriendsAndChat/Public/FriendsAndChatStyle.h
@@ -7,6 +7,7 @@
 #include "FriendsChatStyle.h"
 #include "FriendsChatChromeStyle.h"
 #include "FriendsFontStyle.h"
+#include "FriendsMarkupStyle.h"
 
 #include "FriendsAndChatStyle.generated.h"
 
@@ -88,6 +89,10 @@ struct FRIENDSANDCHAT_API FFriendsAndChatStyle
 	UPROPERTY()
 	FFriendsChatChromeStyle FriendsChatChromeStyle;
 	FFriendsAndChatStyle& SetFriendsChatChromeStyle(const FFriendsChatChromeStyle& InFriendsChatChromeStyle);
+
+	UPROPERTY()
+	FFriendsMarkupStyle FriendsMarkupStyle;
+	FFriendsAndChatStyle& SetFriendsMarkupStyle(const FFriendsMarkupStyle& InFriendsMarkupStyle);
 };
 
 /** Manages the style which provides resources for the rich text widget. */
diff --git a/Engine/Source/Developer/FriendsAndChat/Public/FriendsMarkupStyle.h b/Engine/Source/Developer/FriendsAndChat/Public/FriendsMarkupStyle.h
new file mode 100644
--- /dev/null
+++ b/Engine/Source/Developer/FriendsAndChat/Public/FriendsMarkupStyle.h
@@ -0,0 +1,52 @@
+// Copyright 1998-2015 Epic Games, Inc. All Rights Reserved.
+
+#pragma once
+
+#include "FriendsMarkupStyle.generated.h"
+
+/**
+ * Interface for the services manager.
+ */
+USTRUCT()
+struct FRIENDSANDCHAT_API FFriendsMarkupStyle
+	: public FSlateWidgetStyle
+{
+	GENERATED_USTRUCT_BODY()
+
+	// Default Constructor
+	FFriendsMarkupStyle() { }
+
+	// Default Destructor
+	virtual ~FFriendsMarkupStyle() { }
+
+	/**
+	 * Override widget style function.
+	 */
+	virtual void GetResources( TArray< const FSlateBrush* >& OutBrushes ) const override { }
+
+	// Holds the widget type name
+	static const FName TypeName;
+
+	/**
+	 * Get the type name.
+	 * @return the type name
+	 */
+	virtual const FName GetTypeName() const override { return TypeName; };
+
+	/**
+	 * Get the default style.
+	 * @return the default style
+	 */
+	static const FFriendsMarkupStyle& GetDefault();
+
+	/** Markup Button style */
+	UPROPERTY()
+	FButtonStyle MarkupButtonStyle;
+	FFriendsMarkupStyle& SetMarkupButtonStyle(const FButtonStyle& ButtonStyle);
+
+	UPROPERTY()
+	FTextBlockStyle MarkupTextStyle;
+	FFriendsMarkupStyle& SetMarkupTextStyle(const FTextBlockStyle& InTextStyle);
+
+};
+
diff --git a/Engine/Source/Developer/FriendsAndChat/Public/Layers/DataAccess/FriendsMessageType.h b/Engine/Source/Developer/FriendsAndChat/Public/Layers/DataAccess/FriendsMessageType.h
--- a/Engine/Source/Developer/FriendsAndChat/Public/Layers/DataAccess/FriendsMessageType.h
+++ b/Engine/Source/Developer/FriendsAndChat/Public/Layers/DataAccess/FriendsMessageType.h
@@ -9,14 +9,15 @@ namespace EChatMessageType
 {
 	enum Type : uint8
 	{
-		Custom = 0,
-		Whisper = 1 << 0,		// Person whisper Item
-		Game = 1 << 1,		// Game Chat Item
-		Global = 1 << 2,		// Global Chat Item
-		Team = 1 << 3,		// Team Chat Item
-		Clan = 1 << 4,		// Clan Chat Item
-		Party = 1 << 5,		// Party Chat Item
-		Invalid = 1 << 6,		// Invalid or max
+		Custom = 0,			// Empty
+		Empty = 1 << 0,		// Person whisper Item
+		Whisper = 1 << 1,		// Person whisper Item
+		Game = 1 << 2,		// Game Chat Item
+		Global = 1 << 3,		// Global Chat Item
+		Team = 1 << 4,		// Team Chat Item
+		Clan = 1 << 5,		// Clan Chat Item
+		Party = 1 << 6,		// Party Chat Item
+		Invalid = 1 << 7,		// Invalid or max
 	};
 
 	/** @return the FTextified version of the enum passed in */
@@ -70,6 +71,7 @@ namespace EChatMessageType
 
 	inline EChatMessageType::Type EnumFromString(FString ShortcutString)
 	{
+		static FString EmptyShortcut = TEXT("/");
 		static FString CustomShortcut = TEXT("/c");
 		static FString WhisperShortcut = TEXT("/w");
 		static FString GameShortcut = TEXT("/i");
@@ -106,6 +108,10 @@ namespace EChatMessageType
 		{
 			return EChatMessageType::Party;
 		}
+		else if( ShortcutString == EmptyShortcut)
+		{
+			return EChatMessageType::Empty;
+		}
 		return EChatMessageType::Invalid;
 	}
 };
\ No newline at end of file
diff --git a/Engine/Source/Developer/FriendsAndChat/Public/Models/IFriendItem.h b/Engine/Source/Developer/FriendsAndChat/Public/Models/IFriendItem.h
--- a/Engine/Source/Developer/FriendsAndChat/Public/Models/IFriendItem.h
+++ b/Engine/Source/Developer/FriendsAndChat/Public/Models/IFriendItem.h
@@ -55,6 +55,25 @@ namespace EFriendActionType
 		}
 	}
 
+	inline FString ToAnalyticString(EFriendActionType::Type State)
+	{
+		static const FString FriendRequestIgnored = TEXT("Ignored");
+		static const FString FriendRequestRejected = TEXT("Rejected");
+		static const FString FriendRequestCancelled = TEXT("Cancelled");
+		static const FString FriendBlocked = TEXT("Blocked");
+		static const FString FriendRemoved = TEXT("Removed");
+
+		switch (State)
+		{
+			case IgnoreFriendRequest: return FriendRequestIgnored;
+			case RejectFriendRequest: return FriendRequestRejected;
+			case CancelFriendRequest: return FriendRequestCancelled;
+			case BlockFriend: return FriendBlocked;
+
+			default: return FriendRemoved;
+		}
+	}
+
 	inline EFriendActionLevel ToActionLevel(EFriendActionType::Type State)
 	{
 		switch (State)
@@ -207,16 +226,32 @@ public:
 	/** Get if is from a game request. */
 	virtual bool IsGameRequest() const = 0;
 
+	// Hacks Nick Davies - cache some states so the UI can be regenerated correctly after a refresh
+
+	/** Set if this friend has a pending action. */
 	virtual void SetPendingAction(EFriendActionType::Type InPendingAction)
 	{
 		PendingActionType = InPendingAction;
 	}
 
+	/** Get the current pending action. */
 	virtual EFriendActionType::Type GetPendingAction()
 	{
 		return PendingActionType;
 	}
 
+	/** Set if this friend is online from the last know presence. */
+	virtual void SetChachedOnline(bool InIsOnline)
+	{
+		CachedOnline = InIsOnline;
+	}
+
+	/** Get if this friend is online so we can check if it differs from received presence. */
+	virtual bool IsCachedOnline()
+	{
+		return CachedOnline;
+	}
+
 	/**
 	 * Get the invitation status.
 	 *
@@ -226,10 +261,12 @@ public:
 
 	IFriendItem()
 		: PendingActionType(EFriendActionType::MAX_None)
+		, CachedOnline(false)
 	{}
 
 private:
 	EFriendActionType::Type PendingActionType;
+	bool CachedOnline;
 public:
 
 	/** Virtual destructor. */
