[INTEGRATE] merge UE4-LauncherDev > UE4 CL 2610916
Add messages for invalid / no friends. Fixed up some issues with the slash command button

diff --git a/Engine/Source/Developer/FriendsAndChat/Private/Layers/Domain/FriendsChatMarkupService.cpp b/Engine/Source/Developer/FriendsAndChat/Private/Layers/Domain/FriendsChatMarkupService.cpp
--- a/Engine/Source/Developer/FriendsAndChat/Private/Layers/Domain/FriendsChatMarkupService.cpp
+++ b/Engine/Source/Developer/FriendsAndChat/Private/Layers/Domain/FriendsChatMarkupService.cpp
@@ -35,6 +35,7 @@ struct FProcessedInput
 		MatchedFriends.Empty();
 		NeedsTip = false;
 		FoundMatch = false;
+		ChatChannel = EChatMessageType::Invalid;
 	}
 };
 
@@ -119,7 +120,11 @@ class FInvalidChatTip : public IChatTip
 public:
 	virtual FText GetTipText() override
 	{
-		return LOCTEXT("InvalidChatTip", "Invalid markup");
+		if(InvalidReason.IsEmpty())
+		{
+			return LOCTEXT("InvalidChatTip", "Invalid markup");
+		}
+		return InvalidReason;
 	}
 
 	virtual bool IsValidForType(EChatMessageType::Type ChatChannel) override
@@ -137,10 +142,18 @@ public:
 		return FReply::Handled();
 	}
 
+	FInvalidChatTip(FText InInvalidReason)
+	 : InvalidReason(InInvalidReason)
+	{}
+
 	FInvalidChatTip()
+	 : InvalidReason(FText::GetEmpty())
 	{}
 
 	virtual ~FInvalidChatTip(){}
+
+private:
+	FText InvalidReason;
 };
 
 
@@ -154,7 +167,7 @@ public:
 		ForceDisplayToolTips = !ForceDisplayToolTips;
 		if(ForceDisplayToolTips)
 		{
-			GenerateTips();
+			GenerateDefaultTips();
 		}
 	}
 
@@ -360,9 +373,16 @@ private:
 		ChatTipArray.Empty();
 		if(ProcessedInput->ChatChannel == EChatMessageType::Whisper)
 		{
-			for(const auto& Friend : ProcessedInput->ValidFriends)
+			if(ProcessedInput->ValidFriends.Num())
+			{
+				for(const auto& Friend : ProcessedInput->ValidFriends)
+				{
+					ChatTipArray.Add(MakeShareable(new FFriendChatTip(SharedThis(this), Friend.ToSharedRef())));
+				}
+			}
+			else
 			{
-				ChatTipArray.Add(MakeShareable(new FFriendChatTip(SharedThis(this), Friend.ToSharedRef())));
+				ChatTipArray.Add(MakeShareable(new FInvalidChatTip(LOCTEXT("InvalidChatTip_NoFriends", "No matching friends currently online"))));
 			}
 		}
 		else if(ProcessedInput->ChatChannel != EChatMessageType::Invalid)
@@ -388,6 +408,19 @@ private:
 		OnInputUpdated().Broadcast();
 	}
 
+	void GenerateDefaultTips()
+	{
+		SelectedChatTip.Reset();
+		ChatTipArray.Empty();
+		for(const auto& ChatTip : CommmonChatTips)
+		{
+			ChatTipArray.Add(ChatTip); 
+		}
+		SelectedChatTip = ChatTipArray[0];
+		SelectedChatIndex = 0;
+		OnInputUpdated().Broadcast();
+	}
+
 	void SetNextCommand()
 	{
 		if(SelectedChatTip.IsValid() && ChatTipArray.Num() > 0)
@@ -424,6 +457,7 @@ private:
 			{
 				SetSelectedFriend(nullptr);
 				SetValidatedText(InPlainText);
+				NeedsValidation = true;
 			}
 			else
 			{
