WIP update some styling, and hide chat while in-game. Also release mouse capture

diff --git a/Engine/Source/Developer/FriendsAndChat/Private/Layers/Presentation/ChatDisplayService.cpp b/Engine/Source/Developer/FriendsAndChat/Private/Layers/Presentation/ChatDisplayService.cpp
--- a/Engine/Source/Developer/FriendsAndChat/Private/Layers/Presentation/ChatDisplayService.cpp
+++ b/Engine/Source/Developer/FriendsAndChat/Private/Layers/Presentation/ChatDisplayService.cpp
@@ -78,6 +78,12 @@ public:
 		return FriendsSendNetworkMessageEvent;
 	}
 
+	DECLARE_DERIVED_EVENT(FChatDisplayServiceImpl, IChatDisplayService::FOnFocusReleasedEvent, FOnFocusReleasedEvent);
+	virtual FOnFocusReleasedEvent& OnFocuseReleasedEvent() override
+	{
+		return OnFocusReleasedEvent;
+	}
+
 	DECLARE_DERIVED_EVENT(FChatDisplayServiceImpl, IChatDisplayService::FChatListSetFocus, FChatListSetFocus);
 	virtual FChatListSetFocus& OnChatListSetFocus() override
 	{
@@ -96,6 +102,7 @@ private:
 		else
 		{
 			ChatEntryVisibility = EVisibility::Hidden;
+			OnFocuseReleasedEvent().Broadcast();
 		}
 	}
 
@@ -187,6 +194,8 @@ private:
 	FChatListUpdated ChatListUpdatedEvent;
 	FOnFriendsChatMessageCommitted ChatMessageCommittedEvent;
 	FOnFriendsSendNetworkMessageEvent FriendsSendNetworkMessageEvent;
+	FOnFocusReleasedEvent OnFocusReleasedEvent;
+
 	FChatListSetFocus ChatSetFocusEvent;
 
 	// Delegate for which function we should use when we tick
diff --git a/Engine/Source/Developer/FriendsAndChat/Public/Layers/Presentation/IChatDisplayService.h b/Engine/Source/Developer/FriendsAndChat/Public/Layers/Presentation/IChatDisplayService.h
--- a/Engine/Source/Developer/FriendsAndChat/Public/Layers/Presentation/IChatDisplayService.h
+++ b/Engine/Source/Developer/FriendsAndChat/Public/Layers/Presentation/IChatDisplayService.h
@@ -58,6 +58,9 @@ public:
 	DECLARE_EVENT_OneParam(IChatDisplayService, FOnFriendsSendNetworkMessageEvent, /*struct*/ const FString& /* the message */)
 	virtual FOnFriendsSendNetworkMessageEvent& OnNetworkMessageSentEvent() = 0;
 
+	DECLARE_EVENT(IChatDisplayService, FOnFocusReleasedEvent)
+	virtual FOnFocusReleasedEvent& OnFocuseReleasedEvent() = 0;
+
 	DECLARE_EVENT(IChatDisplayService, FChatListUpdated)
 	virtual FChatListUpdated& OnChatListUpdated() = 0;
 
