[INTEGRATE] merge UE4-LauncherDev > UE4 Friends and Chat up to CL 2633708

diff --git a/Engine/Source/Developer/FriendsAndChat/Private/Layers/DataAccess/FriendsAndChatManager.h b/Engine/Source/Developer/FriendsAndChat/Private/Layers/DataAccess/FriendsAndChatManager.h
--- a/Engine/Source/Developer/FriendsAndChat/Private/Layers/DataAccess/FriendsAndChatManager.h
+++ b/Engine/Source/Developer/FriendsAndChat/Private/Layers/DataAccess/FriendsAndChatManager.h
@@ -2,26 +2,6 @@
 
 #pragma once
 
-
-// Enum holding the state of the Friends manager
-namespace EFriendsAndManagerState
-{
-	enum Type
-	{
-		Idle,								// Idle - can accept requests
-		RequestingFriendsList,				// Requesting a list refresh
-		RequestFriendsListRefresh,			// List request in progress
-		RequestingRecentPlayersIDs,			// Requesting recent player ids
-		RequestRecentPlayersListRefresh,	// Recent players request in progress
-		RequestGameInviteRefresh,			// Game invites can be from non-friends. refresh non-friend user info
-		ProcessFriendsList,					// Process the Friends List after a list refresh
-		RequestingFriendName,				// Requesting a friend add
-		DeletingFriends,					// Deleting a friend
-		AcceptingFriendRequest,				// Accepting a friend request
-		OffLine,							// No logged in
-	};
-};
-
 /**
  * Result of find friend attempt
  */
@@ -94,610 +74,24 @@ public:
 	virtual TSharedPtr< SWidget > GenerateChatWidget(const FFriendsAndChatStyle* InStyle, TAttribute<FText> ActivationHintDelegate, EChatMessageType::Type MessageType, TSharedPtr<IFriendItem> FriendItem, TSharedPtr< SWindow > WindowPtr) override;
 	virtual TSharedPtr<SWidget> GenerateFriendUserHeaderWidget(TSharedPtr<IFriendItem> FriendItem) override;
 	virtual TSharedPtr<class FFriendsNavigationService> GetNavigationService() override;
+	virtual TSharedPtr<class IChatNotificationService> GetNotificationService() override;
+	virtual TSharedPtr<class IChatCommunicationService> GetCommunicationService() override;
+	virtual TSharedPtr<class IGameAndPartyService> GetGameAndPartyService() override;
 	virtual void InsertNetworkChatMessage(const FString& InMessage) override;
 	virtual void JoinPublicChatRoom(const FString& RoomName) override;
 	virtual void OnChatPublicRoomJoined(const FString& ChatRoomID) override;
 
-	/**
-	 * Get the name of the first global chat room we've joined
-	 *
-	 * param Out global chat room name
-	 */
-	virtual bool GetGlobalChatRoomId(FString& OutGlobalRoomId) const;
-
-	/**
-	* Get the name of the primary party chat room we've joined
-	*
-	* @param Out shared ref to party room id
-	*/
-	virtual TSharedPtr<const FOnlinePartyId> GetPartyChatRoomId() const;
-
-	/**
-	* Get the friends filtered list of friends.
-	*
-	* @param OutFriendsList  Array of friends to fill in.
-	* @return the friend list count.
-	*/
-	virtual int32 GetFilteredFriendsList(TArray< TSharedPtr< IFriendItem > >& OutFriendsList) override;
-
-	/**
-	* Get the recent players list.
-	* @return the list.
-	*/
-	virtual TArray< TSharedPtr< IFriendItem > >& GetRecentPlayerList() override;
-
-	/**
-	* Get incoming game invite list.
-	*
-	* @param OutFriendsList  Array of friends to fill in.
-	* @return The friend list count.
-	*/
-	virtual int32 GetFilteredGameInviteList(TArray< TSharedPtr< IFriendItem > >& OutFriendsList) override;
-	
-	/**
-	 * Get the analytics for recording friends chat events
-	 */
-	TSharedPtr<class FFriendsAndChatAnalytics> GetAnalytics()
-	{
-		return Analytics;
-	}
-
-	/**
-	 * Get session id if the current player is in a session.
-	 *
-	 * @return True id of game session we are in.
-	 */
-	TSharedPtr<const FUniqueNetId> GetGameSessionId() const;
-
-	/**
-	 * Get session id from a string
-	 *
-	 * @return The id of game session we are in.
-	 */
-	TSharedPtr<const FUniqueNetId> GetGameSessionId(FString SessionID) const;
-
-	/**
-	 * Get if the current player is in a session.
-	 *
-	 * @return True if we are in a game session.
-	 */
-	bool IsInGameSession() const;
-
-	/**
-	* Get if the current player is in an active party, for determining if party chat is accessible.
-	*
-	* @return True if we are in a party for chat purposes
-	*/
-	bool IsInActiveParty() const;
-
-	/**
-	 * Is this friend in the same session as I am
-	 *
-	 * @param reference to a friend item
-	 */
-	bool IsFriendInSameSession( const TSharedPtr< const IFriendItem >& FriendItem ) const;
-
-	/**
-	 * Is this friend in the same party as I am
-	 *
-	 * @param reference to a friend item
-	 */
-	bool IsFriendInSameParty( const TSharedPtr< const IFriendItem >& FriendItem ) const;
-
-	/**
-	 * Obtain the party join info for a friend entry if avaialble
-	 *
-	 * @param reference to a friend item
-	 *
-	 * @return party info or null if not avail
-	 */
-	TSharedPtr<IOnlinePartyJoinInfo> GetPartyJoinInfo(const TSharedPtr<const IFriendItem>& FriendItem) const;
-
-	/**
-	 * Get if the current player is in a party and is joinable.
-	 *
-	 * @return true if the local user considers their party joinable, otherwise false
-	 */
-	
-	bool IsInJoinableParty() const;
-
-	/**
-	 * Get if the current player is in a session and that game is joinable.
-	 * The context here is from the local perspective, it should not be interpreted by or given to external recipients
-	 * (ie, the game is invite only and therefore joinable via invite, but it is not "joinable" from another user's perspective)
-	 *
-	 * @return true if the local user considers their game joinable, otherwise false
-	 */
-	bool IsInJoinableGameSession() const;
-
-	/**
-	 * @param ClientID ID of the Game we want to join
-	 * @return true if joining a game is allowed
-	 */
-	bool JoinGameAllowed(FString ClientID);
-
-	/**
-	 * @return true if in the launcher
-	 */
-	const bool IsInLauncher() const;
-
-	/** 
-	 * @return if we are logged into global chat
-	 */
-	bool IsInGlobalChat() const;
-
-	/**
-	 * Open the global chat window
-	 */
-	void OpenGlobalChat();
-
-	/**
-	 * Open the chat settings for chrome
-	 */
-	void OpenChromeSettings();
-
-	/**
-	 * Accept a friend request.
-	 *
-	 * @param FriendItem The friend to accept.
-	 */
-	void AcceptFriend( TSharedPtr< IFriendItem > FriendItem );
-
-	/**
-	 * Reject a friend request.
-	 *
-	 * @param FriendItem The friend to reject.
-	 */
-	void RejectFriend( TSharedPtr< IFriendItem > FriendItem );
-
-	/**
-	 * Get outgoing request list.
-	 *
-	 * @param OutFriendsList  Array of friends to fill in.
-	 * @return The friend list count.
-	 */
-	int32 GetFilteredOutgoingFriendsList( TArray< TSharedPtr< IFriendItem > >& OutFriendsList );
-
-	/**
-	 * Request a friend be added.
-	 *
-	 * @param FriendName The friend name.
-	 */
-	void RequestFriend( const FText& FriendName );
-
-	/**
-	 * Request a friend to be deleted.
-	 *
-	 * @param FriendItem The friend item to delete.
-	 * @param DeleteReason The reason the friend is being deleted.
-	 */
-	void DeleteFriend(TSharedPtr< IFriendItem > FriendItem, EFriendActionType::Type DeleteReason);
-
-	/**
-	 * Reject a game invite
-	 *
-	 * @param FriendItem friend with game invite info
-	 */
-	void RejectGameInvite(const TSharedPtr<IFriendItem>& FriendItem);
-
-	/**
-	 * Accept a game invite
-	 *
-	 * @param FriendItem friend with game invite info
-	 */
-	void AcceptGameInvite(const TSharedPtr<IFriendItem>& FriendItem);
-
-	/**
-	 * Send a game invite to a friend
-	 *
-	 * @param FriendItem friend to send game invite to
-	 */
-	void SendGameInvite(const TSharedPtr<IFriendItem>& FriendItem);
-
-	/**
-	* Send a game invite to a user by id
-	*
-	* @param UserId user to send game invite to
-	*/
-	void SendGameInvite(const FUniqueNetId& ToUser);
-
-	/**
-	 * Send a party invite to a user by id
-	 *
-	 * @param UserId user to send party invite to
-	 */
-	bool SendPartyInvite(const FUniqueNetId& ToUser);
-
-	/** Send a game invite notification. */
-	void SendGameInviteNotification(const TSharedPtr<IFriendItem>& FriendItem);
-
-	/** 
-	 * Broadcast when a chat message is received - opens the chat window in the launcher. \
-	 *
-	 * @param ChatType	The type of chat message received
-	 * @param FriendItem The friend item if this chat type is whisper
-	 */
-	void SendChatMessageReceivedEvent(EChatMessageType::Type ChatType, TSharedPtr<IFriendItem> FriendItem);
-
-	/**
-	 * Find a user ID.
-	 *
-	 * @param InUserName The user name to find.
-	 * @return The unique ID.
-	 */
-	TSharedPtr< const FUniqueNetId > FindUserID( const FString& InUsername );
-
-	/**
-	 * Set the user online status
-	 *
-	 * @param OnlineState - the online state
-	 */
-	void SetUserIsOnline(EOnlinePresenceState::Type OnlineState);
-
-	/**
-	* Get the owner's client id
-	*
-	* @return client id string (or an empty string if it fails)
-	*/
-	FString GetUserClientId() const;
-
-	/**
-	* Get the owner's display name
-	*
-	* @return display/nickname string (or an empty string if it fails)
-	*/
-	FString GetUserNickname() const;
-
-	/**
-	 * Find a recent player.
-	 *
-	 * @param InUserId The user id to find.
-	 * @return The recent player ID.
-	 */
-	TSharedPtr< IFriendItem > FindRecentPlayer(const FUniqueNetId& InUserID);
-	
-	/**
-	* Get the the local player's unique id used for the OSS
-	*
-	* @return Pointer to unique id.
-	*/
-	TSharedPtr<const FUniqueNetId> GetLocalUserId() const;
-
-	/**
-	 * Find a user.
-	 *
-	 * @param InUserName The user name to find.
-	 * @return The Friend ID.
-	 */
-	TSharedPtr< IFriendItem > FindUser(const FUniqueNetId& InUserID);
-	TSharedPtr< IFriendItem > FindUser(const TSharedRef<const FUniqueNetId>& InUserID);
-
-	TSharedPtr<class FFriendViewModel> GetFriendViewModel(const TSharedPtr<const FUniqueNetId> InUserID, const FText Username);
-	TSharedPtr<class FFriendViewModel> GetFriendViewModel(const FString InUserID, const FText Username);
-
 	// External events
-	DECLARE_DERIVED_EVENT(FFriendsAndChatManager, IFriendsAndChatManager::FOnFriendsNotificationEvent, FOnFriendsNotificationEvent)
-	virtual FOnFriendsNotificationEvent& OnFriendsNotification() override
-	{
-		return FriendsListNotificationDelegate;
-	}
-
-	DECLARE_DERIVED_EVENT(FFriendsAndChatManager, IFriendsAndChatManager::FOnFriendsNotificationActionEvent, FOnFriendsNotificationActionEvent)
-	virtual FOnFriendsNotificationActionEvent& OnFriendsActionNotification() override
-	{
-		return FriendsListActionNotificationDelegate;
-	}
-
-	DECLARE_DERIVED_EVENT(IFriendsAndChatManager, IFriendsAndChatManager::FOnFriendsUserSettingsUpdatedEvent, FOnFriendsUserSettingsUpdatedEvent)
-	virtual FOnFriendsUserSettingsUpdatedEvent& OnFriendsUserSettingsUpdated() override
-	{
-		return FriendsUserSettingsUpdatedDelegate;
-	}
-
-	DECLARE_DERIVED_EVENT(IFriendsAndChatManager, IFriendsAndChatManager::FOnFriendsJoinGameEvent, FOnFriendsJoinGameEvent)
-	virtual FOnFriendsJoinGameEvent& OnFriendsJoinGame() override
-	{
-		return FriendsJoinGameEvent;
-	}
-
-	DECLARE_DERIVED_EVENT(IFriendsAndChatManager, IFriendsAndChatManager::FOnFriendsJoinPartyEvent, FOnFriendsJoinPartyEvent)
-	virtual FOnFriendsJoinPartyEvent& OnFriendsJoinParty() override
-	{
-		return FriendsJoinPartyEvent;
-	}
-
-	DECLARE_DERIVED_EVENT(IFriendsAndChatManager, IFriendsAndChatManager::FChatMessageReceivedEvent, FChatMessageReceivedEvent);
-	virtual FChatMessageReceivedEvent& OnChatMessageRecieved() override
-	{
-		return ChatMessageReceivedEvent;
-	}
-
 	virtual FAllowFriendsJoinGame& AllowFriendsJoinGame() override
 	{
 		return AllowFriendsJoinGameDelegate;
 	}
 
-	// Internal events
-
-	DECLARE_DERIVED_EVENT(IFriendsAndChatManager, IFriendsAndChatManager::FOnFriendsUpdated, FOnFriendsUpdated)
-	virtual FOnFriendsUpdated& OnFriendsListUpdated()
-	{
-		return OnFriendsListUpdatedDelegate;
-	}
-
-	DECLARE_DERIVED_EVENT(IFriendsAndChatManager, IFriendsAndChatManager::FOnGameInvitesUpdated, FOnGameInvitesUpdated)
-	virtual FOnGameInvitesUpdated& OnGameInvitesUpdated()
-	{
-		return OnGameInvitesUpdatedDelegate;
-	}
-
-	DECLARE_EVENT_OneParam(FFriendsAndChatManager, FOnChatFriendSelected, TSharedPtr<IFriendItem> /*ChatFriend*/)
-	virtual FOnChatFriendSelected& OnChatFriendSelected()
-	{
-		return OnChatFriendSelectedDelegate;
-	}
-
 private:
 
-	/** Process the friends list to send out invites */
-	void SendFriendRequests();
-
-	/**
-	 * Request a list read
-	 */
-	void RequestListRefresh();
-
-	/**
-	 * Request recent player list
-	 */
-	void RequestRecentPlayersListRefresh();
-
-	/**
-	 * Pre process the friends list - find missing names etc
-	 *
-	 * @param ListName - the list name
-	 */
-	void PreProcessList(const FString& ListName);
-
-	/** Refresh the data lists */
-	void RefreshList();
-
-	/**
-	 * Set the manager state.
-	 *
-	 * @param NewState The new manager state.
-	 */
-	void SetState( EFriendsAndManagerState::Type NewState );
-
-	/** Send a friend invite notification. */
-	void SendFriendInviteNotification();
-
-	/** Send a friend invite accepted notification. */
-	void SendInviteAcceptedNotification(TSharedPtr< IFriendItem > Friend);
-
-	/** Send a friend message received notification. */
-	void SendMessageReceivedNotification(TSharedPtr< IFriendItem > Friend);
-
-	/** Get Invite Notification Text */
-	const FText GetInviteNotificationText(TSharedPtr< IFriendItem > Friend) const;
-
 	/** Called when singleton is released. */
 	void ShutdownManager();
 
-	/**
-	 * Delegate used when the friends read request has completed.
-	 *
-	 * @param LocalPlayer		The controller number of the associated user that made the request.
-	 * @param bWasSuccessful	true if the async action completed without error, false if there was an error.
-	 * @param ListName			Name of the friends list that was operated on.
-	 * @param ErrorStr			String representing the error condition.
-	 */
-	void OnReadFriendsListComplete( int32 LocalPlayer, bool bWasSuccessful, const FString& ListName, const FString& ErrorStr );
-
-	/**
-	 * Delegate used when the query for recent players has completed
-	 *
-	 * @param UserId the id of the user that made the request
-	 * @param Namespace the recent players namespace to retrieve
-	 * @param bWasSuccessful true if the async action completed without error, false if there was an error
-	 * @param Error string representing the error condition
-	 */
-	void OnQueryRecentPlayersComplete(const FUniqueNetId& UserId, const FString& Namespace, bool bWasSuccessful, const FString& ErrorStr);
-
-	/**
-	 * Delegate used when an invite send request has completed.
-	 *
-	 * @param LocalPlayer		The controller number of the associated user that made the request.
-	 * @param bWasSuccessful	true if the async action completed without error, false if there was an error.
-	 * @param FriendId			Player that was invited.
-	 * @param ListName			Name of the friends list that was operated on.
-	 * @param ErrorStr			String representing the error condition.
-	 */
-	void OnSendInviteComplete( int32 LocalPlayer, bool bWasSuccessful, const FUniqueNetId& FriendId, const FString& ListName, const FString& ErrorStr );
-
-	/**
-	 * Delegate used when an invite accept request has completed.
-	 *
-	 * @param LocalPlayer		The controller number of the associated user that made the request
-	 * @param bWasSuccessful	true if the async action completed without error, false if there was an error
-	 * @param FriendId			Player that invited us
-	 * @param ListName			Name of the friends list that was operated on
-	 * @param ErrorStr			String representing the error condition
-	 */
-	void OnAcceptInviteComplete( int32 LocalPlayer, bool bWasSuccessful, const FUniqueNetId& FriendId, const FString& ListName, const FString& ErrorStr );
-
-	/**
-	 * Delegate used when the friends delete request has completed.
-	 *
-	 * @param LocalPlayer		The controller number of the associated user that made the request.
-	 * @param bWasSuccessful	true if the async action completed without error, false if there was an error.
-	 * @param DeletedFriendID	The ID of the deleted friend.
-	 * @param ListName			Name of the friends list that was operated on.
-	 * @param ErrorStr			String representing the error condition.
-	 */
-	void OnDeleteFriendComplete( int32 LocalPlayer, bool bWasSuccessful, const FUniqueNetId& DeletedFriendID, const FString& ListName, const FString& ErrorStr );
-
-	/**
-	 * Delegate used when a friend ID request has completed.
-	 *
-	 * @param bWasSuccessful		true if the async action completed without error, false if there was an error.
-	 * @param RequestingUserId		The ID of the user making the request.
-	 * @param DisplayName			The string ID of the friend found.
-	 * @param IdentifiedUserId		The net ID of the found friend.
-	 * @param ErrorStr				String representing the error condition.
-	 */
-	void OnQueryUserIdMappingComplete(bool bWasSuccessful, const FUniqueNetId& RequestingUserId, const FString& DisplayName, const FUniqueNetId& IdentifiedUserId, const FString& Error);
-
-	/**
-	 * Delegate used when a query user info request has completed.
-	 *
-	 * @param LocalPlayer		The controller number of the associated user that made the request.
-	 * @param bWasSuccessful	true if the async action completed without error, false if there was an error.
-	 * @param UserIds			An array of user IDs filled out with info.
-	 * @param ErrorStr			String representing the error condition.
-	 */
-	void OnQueryUserInfoComplete( int32 LocalPlayer, bool bWasSuccessful, const TArray< TSharedRef<const FUniqueNetId> >& UserIds, const FString& ErrorStr );
-
-	/**
-	 * Delegate used when a query presence info request has completed.
-	 *
-	 * @param UserId		The user ID.
-	 * @param Presence	The user presence.
-	 */
-	void OnPresenceReceived(const class FUniqueNetId& UserId, const TSharedRef<FOnlineUserPresence>& NewPresence);
-
-	/**
-	 * Delegate used when a users presence is updated.
-	 *
-	 * @param UserId		The user ID.
-	 * @param Presence	The user presence.
-	 */
-	void OnPresenceUpdated(const class FUniqueNetId& UserId, const bool bWasSuccessful);
-
-	/**
-	 * Delegate called when the friends list changes.
-	 */
-	void OnFriendsListChanged();
-
-	/**
-	 * Delegate called when an invite is received.
-	 *
-	 * @param UserId		The user ID.
-	 * @param FriendId	The friend ID.
-	 */
-	void OnFriendInviteReceived(const FUniqueNetId& UserId, const FUniqueNetId& FriendId);
-
-	/**
-	 * Delegate called when an invite is received to a game session
-	 *
-	 * @param UserId user that received the game invite
-	 * @param FromId user that send the game invite
-	 * @param InviteResult info about the game session that can be joined
-	 */
-	void OnGameInviteReceived(const FUniqueNetId& UserId, const FUniqueNetId& FromId, const FString& ClientId, const FOnlineSessionSearchResult& InviteResult);
-
-	/**
-	 * Delegate callback when an invite is received to a party session
-	 *
-	 * @param RecipientId - id of user that this invite is for
-	 * @param PartyId - id associated with the party
-	 * @param SenderId - id of member that sent the invite
-	 */
-	void OnPartyInviteReceived(const FUniqueNetId& RecipientId, const FOnlinePartyId& PartyId, const FUniqueNetId& SenderId);
-
-	/**
-	 * Process any invites that were received and are pending. 
-	 * Will remove entries that are processed
-	 */
-	void ProcessReceivedGameInvites();
-
-	/**
-	* Process any invites that came in from the command line and add them to received invites
-	*/
-	void ProcessCommandLineInvites();
-
-	/**
-	 * Query for user info associated with user invites
-	 *
-	 * @return true if request is pending
-	 */
-	bool RequestGameInviteUserInfo();
-
-	/**
-	 * Detect cleanup of game session
-	 */
-	void OnGameDestroyed(const FName SessionName, bool bWasSuccessful);
-
-	/**
-	 * Handle switching default chat channel when player joins party
-	 *
-	 * @param PartyId party that member joined
-	 * @param MemberId member that joined
-	 */
-	void OnPartyMemberJoined(const FUniqueNetId& LocalUserId, const FOnlinePartyId& PartyId, const FUniqueNetId& MemberId);
-
-	/**
-	 * Handle removing party chat channel option when party becomes empty
-	 *
-	 * @param PartyId party that member left
-	 * @param MemberId member that left
-	 */
-	void OnPartyMemberExited(const FUniqueNetId& LocalUserId, const FOnlinePartyId& PartyId, const FUniqueNetId& MemberId, const EMemberExitedReason Reason);
-	
-	/**
-	* Handle adding party chat channel when joining another player's party
-	*
-	* @param PartyId party that you joined
-	*/
-	void OnPartyJoined(const FUniqueNetId& LocalUserId, const FOnlinePartyId& PartyId);
-
-	/**
-	 * Delegate used when a friend is removed.
-	 *
-	 * @param UserId		The user ID.
-	 * @param FriendId	The friend ID.
-	 */
-	void OnFriendRemoved(const FUniqueNetId& UserId, const FUniqueNetId& FriendId);
-
-	/**
-	 * Delegate used when an invite is rejected.
-	 *
-	 * @param UserId		The user ID.
-	 * @param FriendId	The friend ID.
-	 */
-	void OnInviteRejected(const FUniqueNetId& UserId, const FUniqueNetId& FriendId);
-
-	/**
-	 * Delegate used when an invite is accepted.
-	 *
-	 * @param UserId		The user ID.
-	 * @param FriendId	The friend ID.
-	 */
-	void OnInviteAccepted(const FUniqueNetId& UserId, const FUniqueNetId& FriendId);
-
-	/**
-	 * Handle an accept message accepted from a notification.
-	 *
-	 * @param MessageNotification The message responded to.
- 	 * @param ResponseTpe The response from the message request.
-	 */
-	FReply HandleMessageAccepted( TSharedPtr< FFriendsAndChatMessage > MessageNotification, EFriendsResponseType::Type ResponseType );
-
-	/**
-	 * Handle an join game message accepted from a notification.
-	 *
-	 * @param MessageNotification The message responded to.
- 	 * @param ResponseTpe The response from the message request.
-	 */
-	FReply HandleGameInviteResponse( TSharedPtr< FFriendsAndChatMessage > MessageNotification, EFriendsResponseType::Type ResponseType );
-
-	/**
-	 * Validate friend list after a refresh. Check if the list has changed.
-	 *
-	 * @return True if the list has changes and the UI needs refreshing.
-	 */
-	bool ProcessFriendsList();
-
 	/**
 	 * A ticker used to perform updates on the main thread.
 	 *
@@ -710,7 +104,7 @@ private:
 	void AddFriendsToast(const FText Message);
 
 	/**
-	 * Get the best spawn position for the next chat windo
+	 * Get the best spawn position for the next chat window
 	 *
 	 * @return Spawn position
 	 */
@@ -718,256 +112,56 @@ private:
 
 private:
 
-	/* Lists used to keep track of Friends
-	*****************************************************************************/
-
-	// Holds the list of user names to query add as friends 
-	TArray< FString > FriendByNameRequests;
-	// Holds the list of user names to send invites to. 
-	TArray< FString > FriendByNameInvites;
-	// Holds the list of Unique Ids found for user names to add as friends
-	TArray< TSharedRef< const FUniqueNetId > > QueryUserIds;
-	// Holds the full friends list used to build the UI
-	TArray< TSharedPtr< IFriendItem > > FriendsList;
-	// Holds the recent players list
-	TArray< TSharedPtr< IFriendItem > > RecentPlayersList;
-	// Holds the filtered friends list used in the UI
-	TArray< TSharedPtr< IFriendItem > > FilteredFriendsList;
-	// Holds the list of old presence data for comparison to new presence updates
-	TMap< FString, FOnlineUserPresence > OldUserPresenceMap;
-	// Holds the outgoing friend request list used in the UI
-	TArray< TSharedPtr< IFriendItem > > FilteredOutgoingList;
-	// Holds the unprocessed friends list generated from a friends request update
-	TArray< TSharedPtr< IFriendItem > > PendingFriendsList;
-	// Holds the list of incoming invites that need to be responded to
-	TArray< TSharedPtr< IFriendItem > > PendingIncomingInvitesList;
-	// Holds the list of incoming game invites that need to be responded to
-	TMap< FString, TSharedPtr< IFriendItem > > PendingGameInvitesList;
-	// Holds the list of invites we have already responded to
-	TArray< TSharedPtr< const FUniqueNetId > > NotifiedRequest;
 	// Holds the list messages sent out to be responded to
 	TArray<TSharedPtr< FFriendsAndChatMessage > > NotificationMessages;
-	// Holds an array of outgoing invite friend requests
-	TArray< TSharedRef< const FUniqueNetId > > PendingOutgoingFriendRequests;
-	// Holds an array of outgoing delete friend requests
-	TArray< FString > PendingOutgoingDeleteFriendRequests;
-	// Holds an array of outgoing accept friend requests
-	TArray< FString > PendingOutgoingAcceptFriendRequests;
-
-	/**
-	 * Game invite that needs to be processed before being displayed
-	 */
-	class FReceivedGameInvite
-	{
-	public:
-		FReceivedGameInvite(
-			const TSharedRef<const FUniqueNetId>& InFromId,
-			const FOnlineSessionSearchResult& InInviteResult,
-			const FString& InClientId)
-			: FromId(InFromId)
-			, InviteResult(new FOnlineSessionSearchResult(InInviteResult))
-			, ClientId(InClientId)
-		{}
-		// who sent the invite (could be non-friend)
-		TSharedRef<const FUniqueNetId> FromId;
-		// session info needed to join the invite
-		TSharedRef<FOnlineSessionSearchResult> InviteResult;
-		// Client id for user that sent the invite
-		FString ClientId;
-		// equality check
-		bool operator==(const FReceivedGameInvite& Other) const
-		{
-			return *Other.FromId == *FromId && Other.ClientId == ClientId;
-		}
-	};
-	// List of game invites that need to be processed
-	TArray<FReceivedGameInvite> ReceivedGameInvites;
-
-	/**
-	 * Party invite that needs to be processed before being displayed
-	 */
-	class FReceivedPartyInvite
-	{
-	public:
-		FReceivedPartyInvite(
-			const TSharedRef<const FUniqueNetId>& InFromId,
-			const TSharedRef<IOnlinePartyJoinInfo>& InPartyJoinInfo)
-			: FromId(InFromId)
-			, PartyJoinInfo(InPartyJoinInfo)
-		{}
-		// who sent the invite (could be non-friend)
-		TSharedRef<const FUniqueNetId> FromId;
-		// info needed to join the party
-		TSharedRef<IOnlinePartyJoinInfo> PartyJoinInfo;
-		// equality check
-		bool operator==(const FReceivedPartyInvite& Other) const
-		{
-			return *Other.FromId == *FromId || *Other.PartyJoinInfo->GetPartyId() == *PartyJoinInfo->GetPartyId();
-		}
-	};
-	// List of party invites that need to be processed
-	TArray<FReceivedPartyInvite> ReceivedPartyInvites;
 
 	/* Delegates
 	*****************************************************************************/
-
-	// Holds the ticker delegate
-	FTickerDelegate UpdateFriendsTickerDelegate;
-	// Delegate to use for querying for recent players 
-	FOnQueryRecentPlayersCompleteDelegate OnQueryRecentPlayersCompleteDelegate;
-	// Delegate to use for deleting a friend
-	FOnDeleteFriendCompleteDelegate OnDeleteFriendCompleteDelegate;
-	// Delegate for querying user id from a name string
-	IOnlineUser::FOnQueryUserMappingComplete OnQueryUserIdMappingCompleteDelegate;
-	// Delegate to use for querying user info list
-	FOnQueryUserInfoCompleteDelegate OnQueryUserInfoCompleteDelegate;
-	// Delegate to use for querying user presence
-	FOnPresenceReceivedDelegate OnPresenceReceivedCompleteDelegate;
-	// Delegate to owner presence updated
-	IOnlinePresence::FOnPresenceTaskCompleteDelegate OnPresenceUpdatedCompleteDelegate;
-	// Delegate for friends list changing
-	FOnFriendsChangeDelegate OnFriendsListChangedDelegate;
-	// Delegate for an invite received
-	FOnInviteReceivedDelegate OnFriendInviteReceivedDelegate;
-	// Delegate for a game invite received
-	FOnSessionInviteReceivedDelegate OnGameInviteReceivedDelegate;
-	// Delegate for a party invite received
-	FOnPartyInviteReceivedDelegate OnPartyInviteReceivedDelegate;
-	// Delegate for a game session being destroyed
-	FOnDestroySessionCompleteDelegate OnDestroySessionCompleteDelegate;
-	// Delegate for a player joining your party (doesn't fire for self joining)
-	FOnPartyMemberJoinedDelegate OnPartyMemberJoinedDelegate;
-	// Delegate for a player leaving your party (includes yourself)
-	FOnPartyMemberExitedDelegate OnPartyMemberExitedDelegate;
-	// Delegate for joining another player's party
-	FOnPartyJoinedDelegate OnPartyJoinedDelegate;
-
-	// Delegate for friend removed
-	FOnFriendRemovedDelegate OnFriendRemovedDelegate;
-	// Delegate for friend invite rejected
-	FOnInviteRejectedDelegate	OnFriendInviteRejected;
-	// Delegate for friend invite accepted
-	FOnInviteAcceptedDelegate	OnFriendInviteAccepted;
-	// Delegate for chat message received
-	FChatMessageReceivedEvent OnChatmessageRecieved;
-
-	// Holds the Friends list notification delegate
-	FOnFriendsNotificationEvent FriendsListNotificationDelegate;
-	// Holds the Friends Action notification delegate
-	FOnFriendsNotificationActionEvent FriendsListActionNotificationDelegate;
-	// Holds the Options Updated event notification delegate
-	FOnFriendsUserSettingsUpdatedEvent FriendsUserSettingsUpdatedDelegate;
-	// Holds the join game request delegate
-	FOnFriendsJoinGameEvent FriendsJoinGameEvent;
-	// Holds the join party request delegate
-	FOnFriendsJoinPartyEvent FriendsJoinPartyEvent;
-	// Holds the chat message received delegate
-	FChatMessageReceivedEvent ChatMessageReceivedEvent;
-	// Delegate callback for determining if joining games functionality should be allowed
 	FAllowFriendsJoinGame AllowFriendsJoinGameDelegate;
 
-	// Internal events
-	// Holds the delegate to call when the friends list gets updated - refresh the UI
-	FOnFriendsUpdated OnFriendsListUpdatedDelegate;
-	// Delegate for when list of active game invites updates
-	FOnGameInvitesUpdated OnGameInvitesUpdatedDelegate;
-	// Delegate for when a friend is selected for chat
-	FOnChatFriendSelected OnChatFriendSelectedDelegate;
-
 	/* Services
 	*****************************************************************************/
 	// Manages private/public chat messages 
-	TSharedPtr<class FFriendsMessageManager> MessageManager;
+	TSharedPtr<class FMessageService> MessageService;
 	// Manages navigation services
 	TSharedPtr<class FFriendsNavigationService> NavigationService;
+	// Manages notification services
+	TSharedPtr<class FChatNotificationService> NotificationService;
 	// Holds the friends view model facotry
 	TSharedPtr<class IFriendViewModelFactory> FriendViewModelFactory;
 	// Holds the friends list provider
 	TSharedPtr<class IFriendListFactory> FriendsListFactory;
 	// Holds the Markup service provider
 	TSharedPtr<class IFriendsChatMarkupServiceFactory> MarkupServiceFactory;
+	// Holds the Game and Party service
+	TSharedPtr<class FGameAndPartyService> GameAndPartyService;
+	// Holds the friends service
+	TSharedPtr<class FFriendsService> FriendsService;
+	// Holds the OSSScheduler 
+	TSharedPtr<class FOSSScheduler> OSSScheduler;
 
-	/* Identity stuff
-	*****************************************************************************/
 
-	// Holds the Online identity
-	IOnlineIdentityPtr OnlineIdentity;
-	// Holds the Friends Interface
-	IOnlineFriendsPtr FriendsInterface;
-	// Holds the Online Subsystem
-	IOnlineSubsystem* OnlineSub;
 
 	/* Chat stuff
 	*****************************************************************************/
 
 	// Keeps track of global chat rooms that have been requested to join
 	TArray<FString> ChatRoomstoJoin;
-	// Holds the chat view model
-	TSharedPtr<class FChatViewModel> ChatViewModel;
-	// Joined Global Chat
-	bool bJoinedGlobalChat;
-
 	// Holds the clan repository
 	TSharedPtr<class IClanRepository> ClanRepository;
 
 	/* Manger state
 	*****************************************************************************/
 
-	// Holds the manager state
-	EFriendsAndManagerState::Type ManagerState;
-
 	// Holds the Friends List widget
 	TSharedPtr< SWidget > FriendListWidget;
 	
-	// Holds the application view models - used for launching and querying the games
-	TMap<FString, TSharedPtr<IFriendsApplicationViewModel>> ApplicationViewModels;
 	// Holds the style used to create the Friends List widget
 	FFriendsAndChatStyle Style;
-	// Holds if the Friends list is inited
-	bool bIsInited;
-	// Holds the Friends system user settings
-	FFriendsAndChatSettings UserSettings;
-	// Holds if we need a list refresh
-	bool bRequiresListRefresh;
-	// Holds if we need a recent player list refresh
-	bool bRequiresRecentPlayersRefresh;
 	// Holds the toast notification
 	TSharedPtr<SNotificationList> FriendsNotificationBox;
 
-	// Lets us know if we are in game for invites / join game sessions
-	bool bIsInGame;
-
-	// The local player index that we use to talk to OSS
-	int32 LocalControllerIndex;
-
-	/** Party that may have been passed on command line to auto-join */
-	FString AutoAcceptPartyId;
-
-	/** Recent players namespace to query */
-	TArray<FString> RecentPlayersNamespaces;
-	TArray<FString> PendingRecentPlayersQueries;
-
 private:
 
 	TSharedPtr<class FFriendsAndChatAnalytics> Analytics;
-	float FlushChatAnalyticsCountdown;
-
-	/** Handle to various registered delegates */
-	FDelegateHandle OnQueryRecentPlayersCompleteDelegateHandle;
-	FDelegateHandle OnFriendsListChangedDelegateHandle;
-	FDelegateHandle OnFriendInviteReceivedDelegateHandle;
-	FDelegateHandle OnFriendRemovedDelegateHandle;
-	FDelegateHandle OnFriendInviteRejectedHandle;
-	FDelegateHandle OnFriendInviteAcceptedHandle;
-	FDelegateHandle OnReadFriendsCompleteDelegateHandle;
-	FDelegateHandle OnDeleteFriendCompleteDelegateHandle;
-	FDelegateHandle OnQueryUserInfoCompleteDelegateHandle;
-	FDelegateHandle OnPresenceReceivedCompleteDelegateHandle;
-	FDelegateHandle OnGameInviteReceivedDelegateHandle;
-	FDelegateHandle OnPartyInviteReceivedDelegateHandle;
-	FDelegateHandle OnDestroySessionCompleteDelegateHandle;
-	FDelegateHandle OnPartyMemberJoinedDelegateHandle;
-	FDelegateHandle OnPartyMemberExitedDelegateHandle;
-	FDelegateHandle OnPartyJoinedDelegateHandle;
-	FDelegateHandle UpdateFriendsTickerDelegateHandle;
 };
diff --git a/Engine/Source/Developer/FriendsAndChat/Private/Layers/DataAccess/OSSScheduler.cpp b/Engine/Source/Developer/FriendsAndChat/Private/Layers/DataAccess/OSSScheduler.cpp
new file mode 100644
--- /dev/null
+++ b/Engine/Source/Developer/FriendsAndChat/Private/Layers/DataAccess/OSSScheduler.cpp
@@ -0,0 +1,344 @@
+// Copyright 1998-2015 Epic Games, Inc. All Rights Reserved.
+
+#include "FriendsAndChatPrivatePCH.h"
+#include "OSSScheduler.h"
+
+#include "FriendsAndChatAnalytics.h"
+
+const float CHAT_ANALYTICS_INTERVAL = 5 * 60.0f;  // 5 min
+
+class FOSSSchedulerImpl
+	: public FOSSScheduler
+{
+public:
+
+	virtual void Login(IOnlineSubsystem* InOnlineSub, bool bInIsGame, int32 LocalUserID) override
+	{
+		// Clear existing data
+		Logout();
+		
+		LocalControllerIndex = LocalUserID;
+		
+		if (InOnlineSub)
+		{
+			OnlineSub = InOnlineSub;
+		}
+		else
+		{
+			OnlineSub = IOnlineSubsystem::Get(TEXT("MCP"));
+		}
+
+		if (OnlineSub != nullptr &&
+			OnlineSub->GetUserInterface().IsValid() &&
+			OnlineSub->GetIdentityInterface().IsValid())
+		{
+			OnlineIdentity = OnlineSub->GetIdentityInterface();
+			OnPresenceUpdatedCompleteDelegate = IOnlinePresence::FOnPresenceTaskCompleteDelegate::CreateSP(this, &FOSSSchedulerImpl::OnPresenceUpdated);
+		}
+
+		if (UpdateFriendsTickerDelegate.IsBound() == false)
+		{
+			UpdateFriendsTickerDelegate = FTickerDelegate::CreateSP(this, &FOSSSchedulerImpl::Tick);
+		}
+
+		UpdateFriendsTickerDelegateHandle = FTicker::GetCoreTicker().AddTicker(UpdateFriendsTickerDelegate);
+	}
+
+	virtual void Logout() override
+	{
+		OnlineSub = nullptr;
+		OnlineIdentity = nullptr;
+		// flush before removing the analytics provider
+		Analytics->FireEvent_FlushChatStats();
+
+		if (UpdateFriendsTickerDelegate.IsBound())
+		{
+			FTicker::GetCoreTicker().RemoveTicker(UpdateFriendsTickerDelegateHandle);
+		}
+	}
+
+	bool IsLoggedIn() const override
+	{
+		return OnlineSub != nullptr;
+	}
+
+	void SetOnline() override
+	{
+		SetUserIsOnline(EOnlinePresenceState::Online);
+	}
+
+	void SetAway() override
+	{
+		SetUserIsOnline(EOnlinePresenceState::Away);
+	}
+
+	virtual EOnlinePresenceState::Type GetOnlineStatus() override
+	{
+		if (OnlineSub != nullptr &&
+			OnlineIdentity.IsValid())
+		{
+			TSharedPtr<FOnlineUserPresence> Presence;
+			TSharedPtr<const FUniqueNetId> UserId = OnlineIdentity->GetUniquePlayerId(0);
+			if (UserId.IsValid())
+			{
+				OnlineSub->GetPresenceInterface()->GetCachedPresence(*UserId, Presence);
+				if (Presence.IsValid())
+				{
+					return Presence->Status.State;
+				}
+			}
+		}
+		return EOnlinePresenceState::Offline;
+	}
+
+	virtual FString GetUserClientId() const override
+	{
+		FString Result;
+		if (OnlineSub != nullptr &&
+			OnlineIdentity.IsValid())
+		{
+			TSharedPtr<FOnlineUserPresence> Presence;
+			TSharedPtr<const FUniqueNetId> UserId = OnlineIdentity->GetUniquePlayerId(0);
+			if (UserId.IsValid())
+			{
+				GetPresenceInterface()->GetCachedPresence(*UserId, Presence);
+				if (Presence.IsValid())
+				{
+					const FVariantData* ClientId = Presence->Status.Properties.Find(DefaultClientIdKey);
+					if (ClientId != nullptr && ClientId->GetType() == EOnlineKeyValuePairDataType::String)
+					{
+						ClientId->GetValue(Result);
+					}
+				}
+			}
+		}
+		return Result;
+	}
+
+	virtual TSharedPtr<class FFriendsAndChatAnalytics> GetAnalytics() const override
+	{
+		return Analytics;
+	}
+
+	virtual IOnlineFriendsPtr GetFriendsInterface() const override
+	{
+		return OnlineSub->GetFriendsInterface();
+	}
+
+	virtual IOnlineUserPtr GetUserInterface() const override
+	{
+		return OnlineSub->GetUserInterface();
+	}
+
+	virtual IOnlineIdentityPtr GetOnlineIdentity() const override
+	{
+		return OnlineIdentity;
+	}
+
+	virtual IOnlinePartyPtr GetPartyInterface() const override
+	{
+		return OnlineSub->GetPartyInterface();
+	}
+
+	virtual IOnlineSessionPtr GetSessionInterface() const override
+	{
+		return OnlineSub->GetSessionInterface();
+	}
+
+	virtual IOnlinePresencePtr GetPresenceInterface() const override
+	{
+		return OnlineSub->GetPresenceInterface();
+	}
+
+	virtual IOnlineChatPtr GetChatInterface() const override
+	{
+		return OnlineSub->GetChatInterface();
+	}
+
+	virtual TSharedPtr<const FOnlinePartyId> GetPartyChatRoomId() const override
+	{
+		TSharedPtr<const FUniqueNetId> UserId = GetOnlineIdentity()->GetUniquePlayerId(LocalControllerIndex);
+		TSharedPtr<const FOnlinePartyId> PartyChatRoomId;
+		TArray< TSharedRef<const FOnlinePartyId> > OutPartyIds;
+		if (UserId.IsValid())
+		{
+			IOnlinePartyPtr PartyInterface = GetPartyInterface();
+			if (PartyInterface.IsValid()
+				&& PartyInterface->GetJoinedParties(*UserId, OutPartyIds) == true
+				&& OutPartyIds.Num() > 0)
+			{
+				PartyChatRoomId = OutPartyIds[0]; // @todo EN need to identify the primary game party consistently when multiple parties exist
+			}
+		}
+		return PartyChatRoomId;
+	}
+
+	virtual void SetUserIsOnline(EOnlinePresenceState::Type OnlineState) override
+	{
+		if (OnlineSub != nullptr &&
+			OnlineIdentity.IsValid())
+		{
+			TSharedPtr<const FUniqueNetId> UserId = OnlineIdentity->GetUniquePlayerId(0);
+			if (UserId.IsValid())
+			{
+				TSharedPtr<FOnlineUserPresence> CurrentPresence;
+				OnlineSub->GetPresenceInterface()->GetCachedPresence(*UserId, CurrentPresence);
+				FOnlineUserPresenceStatus NewStatus;
+				if (CurrentPresence.IsValid())
+				{
+					NewStatus = CurrentPresence->Status;
+				}
+				NewStatus.State = OnlineState;
+				OnlineSub->GetPresenceInterface()->SetPresence(*UserId, NewStatus, OnPresenceUpdatedCompleteDelegate);
+			}
+		}
+	}
+
+	virtual TSharedPtr<const FUniqueNetId> GetGameSessionId() const override
+	{
+		if (GetOnlineIdentity().IsValid() && GetSessionInterface().IsValid())
+		{
+			const FNamedOnlineSession* GameSession = GetSessionInterface()->GetNamedSession(GameSessionName);
+			if (GameSession != nullptr)
+			{
+				TSharedPtr<FOnlineSessionInfo> UserSessionInfo = GameSession->SessionInfo;
+				if (UserSessionInfo.IsValid())
+				{
+					return GetOnlineIdentity()->CreateUniquePlayerId(UserSessionInfo->GetSessionId().ToString());
+				}
+			}
+		}
+		return nullptr;
+	}
+
+	virtual bool IsFriendInSameSession(const TSharedPtr< const IFriendItem >& FriendItem) const override
+	{
+		TSharedPtr<const FUniqueNetId> MySessionId = GetGameSessionId();
+		bool bMySessionValid = MySessionId.IsValid() && MySessionId->IsValid();
+
+		TSharedPtr<const FUniqueNetId> FriendSessionId = FriendItem->GetGameSessionId();
+		bool bFriendSessionValid = FriendSessionId.IsValid() && FriendSessionId->IsValid();
+
+		return (bMySessionValid && bFriendSessionValid && (*FriendSessionId == *MySessionId));
+	}
+
+	virtual bool IsFriendInSameParty(const TSharedPtr< const IFriendItem >& FriendItem) const override
+	{
+		if (GetOnlineIdentity().IsValid() && GetPartyInterface().IsValid())
+		{
+			TSharedPtr<const FUniqueNetId> UserId = GetOnlineIdentity()->GetUniquePlayerId(LocalControllerIndex);
+			if (UserId.IsValid())
+			{
+				TSharedPtr<IOnlinePartyJoinInfo> FriendPartyJoinInfo = GetPartyJoinInfo(FriendItem);
+				if (FriendPartyJoinInfo.IsValid())
+				{
+					TArray<TSharedRef<const FOnlinePartyId>> JoinedParties;
+					if (GetPartyInterface()->GetJoinedParties(*UserId, JoinedParties))
+					{
+						for (auto PartyId : JoinedParties)
+						{
+							if (*PartyId == *FriendPartyJoinInfo->GetPartyId())
+							{
+								return true;
+							}
+						}
+					}
+				}
+			}
+		}
+
+		return false;
+	}
+
+	virtual TSharedPtr<IOnlinePartyJoinInfo> GetPartyJoinInfo(const TSharedPtr<const IFriendItem>& FriendItem) const override
+	{
+		TSharedPtr<IOnlinePartyJoinInfo> Result;
+
+		if (GetOnlineIdentity().IsValid() && GetPartyInterface().IsValid())
+		{
+			TSharedPtr<const FUniqueNetId> UserId = GetOnlineIdentity()->GetUniquePlayerId(LocalControllerIndex);
+			if (UserId.IsValid())
+			{
+				Result = GetPartyInterface()->GetAdvertisedParty(*UserId, *FriendItem->GetUniqueID());
+			}
+		}
+		return Result;
+	}
+
+	void OnPresenceUpdated(const FUniqueNetId& UserId, const bool bWasSuccessful)
+	{
+	}
+
+	virtual void RequestOSS(FOnRequestOSSAccepted OnRequestOSSAcceptedDelegate) override
+	{
+		RequestDelegates.Add(OnRequestOSSAcceptedDelegate);
+	}
+
+	virtual void ReleaseOSS() override
+	{
+		bOSSInUse = false;
+	}
+
+	bool Tick(float Delta)
+	{
+		FlushChatAnalyticsCountdown -= Delta;
+		if (FlushChatAnalyticsCountdown <= 0)
+		{
+			Analytics->FireEvent_FlushChatStats();
+			// Reset countdown for new update
+			FlushChatAnalyticsCountdown = CHAT_ANALYTICS_INTERVAL;
+		}
+		ProcessOSSRequests();
+		return true;
+	}
+
+	void ProcessOSSRequests()
+	{
+		if (!bOSSInUse && RequestDelegates.Num() > 0)
+		{
+			bOSSInUse = true;
+			RequestDelegates[0].Execute();
+			RequestDelegates.RemoveAt(0);
+		}
+	}
+
+private:
+
+
+	FOSSSchedulerImpl(TSharedPtr<FFriendsAndChatAnalytics> InAnalytics)
+		: OnlineIdentity(nullptr)
+		, OnlineSub(nullptr)
+		, Analytics(InAnalytics)
+		, bOSSInUse(false)
+		, FlushChatAnalyticsCountdown(CHAT_ANALYTICS_INTERVAL)
+	{
+	}
+
+	// Holds the Online identity
+	IOnlineIdentityPtr OnlineIdentity;
+	// Holds the Online Subsystem
+	IOnlineSubsystem* OnlineSub;
+
+	// Delegate to owner presence updated
+	IOnlinePresence::FOnPresenceTaskCompleteDelegate OnPresenceUpdatedCompleteDelegate;
+
+	TSharedPtr<FFriendsAndChatAnalytics> Analytics;
+
+	int32 LocalControllerIndex;
+
+	TArray<FOnRequestOSSAccepted> RequestDelegates;
+	bool bOSSInUse;
+
+	// Holds the ticker delegate
+	FTickerDelegate UpdateFriendsTickerDelegate;
+	FDelegateHandle UpdateFriendsTickerDelegateHandle;
+	float FlushChatAnalyticsCountdown;
+
+	friend FOSSSchedulerFactory;
+};
+
+TSharedRef< FOSSScheduler > FOSSSchedulerFactory::Create(TSharedPtr<FFriendsAndChatAnalytics> Analytics)
+{
+	TSharedRef< FOSSSchedulerImpl > OSSScheduler(new FOSSSchedulerImpl(Analytics));
+	return OSSScheduler;
+}
\ No newline at end of file
diff --git a/Engine/Source/Developer/FriendsAndChat/Private/Layers/DataAccess/OSSScheduler.h b/Engine/Source/Developer/FriendsAndChat/Private/Layers/DataAccess/OSSScheduler.h
new file mode 100644
--- /dev/null
+++ b/Engine/Source/Developer/FriendsAndChat/Private/Layers/DataAccess/OSSScheduler.h
@@ -0,0 +1,51 @@
+// Copyright 1998-2015 Epic Games, Inc. All Rights Reserved.
+
+#pragma once
+
+/**
+* Implement the Friends data access
+*/
+class FOSSScheduler
+	: public TSharedFromThis < FOSSScheduler >
+{
+public:
+
+	/** Destructor. */
+	virtual ~FOSSScheduler() {};
+
+	virtual void Login(IOnlineSubsystem* InOnlineSub, bool bInIsGame, int32 LocalUserID) = 0;
+	virtual void Logout() = 0;
+	virtual bool IsLoggedIn() const = 0;
+	virtual void SetOnline() = 0;
+	virtual void SetAway() = 0;
+	virtual EOnlinePresenceState::Type GetOnlineStatus() = 0;
+	virtual FString GetUserClientId() const = 0;
+	virtual TSharedPtr<const FOnlinePartyId> GetPartyChatRoomId() const = 0;
+	virtual void SetUserIsOnline(EOnlinePresenceState::Type OnlineState) = 0;
+	virtual TSharedPtr<const FUniqueNetId> GetGameSessionId() const = 0;
+	virtual bool IsFriendInSameSession(const TSharedPtr< const IFriendItem >& FriendItem) const = 0;
+	virtual bool IsFriendInSameParty(const TSharedPtr< const IFriendItem >& FriendItem) const = 0;
+	virtual TSharedPtr<IOnlinePartyJoinInfo> GetPartyJoinInfo(const TSharedPtr<const IFriendItem>& FriendItem) const = 0;
+	
+
+	virtual TSharedPtr<class FFriendsAndChatAnalytics> GetAnalytics() const = 0;
+	virtual IOnlineFriendsPtr GetFriendsInterface() const = 0;
+	virtual IOnlineUserPtr GetUserInterface() const = 0;
+	virtual IOnlineIdentityPtr GetOnlineIdentity() const = 0;
+	virtual IOnlinePartyPtr GetPartyInterface() const = 0;
+	virtual IOnlineSessionPtr GetSessionInterface() const = 0;
+	virtual IOnlinePresencePtr GetPresenceInterface() const = 0;
+	virtual IOnlineChatPtr GetChatInterface() const = 0;
+
+	DECLARE_DELEGATE(FOnRequestOSSAccepted);
+
+	virtual void RequestOSS(FOnRequestOSSAccepted OnRequestOSSAcceptedDelegate) = 0;
+	virtual void ReleaseOSS() = 0;
+};
+
+/**
+* Creates the implementation of a chat manager.
+*
+* @return the newly created FriendViewModel implementation.
+*/
+FACTORY(TSharedRef< FOSSScheduler >, FOSSScheduler, TSharedPtr<class FFriendsAndChatAnalytics> Analytics);
diff --git a/Engine/Source/Developer/FriendsAndChat/Private/Layers/Domain/ClanInfo.h b/Engine/Source/Developer/FriendsAndChat/Private/Layers/Domain/ClanInfo.h
--- a/Engine/Source/Developer/FriendsAndChat/Private/Layers/Domain/ClanInfo.h
+++ b/Engine/Source/Developer/FriendsAndChat/Private/Layers/Domain/ClanInfo.h
@@ -8,11 +8,13 @@ public:
 	virtual ~IClanInfo() {}
 
 	virtual FText GetTitle() const = 0;
+	virtual const TSharedRef<const FUniqueNetId> GetUniqueID() const = 0;
 	virtual FText GetDescription() const = 0;
 	virtual FText GetLongDescription() const = 0;
 	virtual TArray<TSharedPtr<class IFriendItem>> GetMemberList() const = 0;
 	virtual FText GetClanBrushName() const = 0;
 	virtual int32 GetClanMembersCount() const = 0;
+	virtual bool IsPrimaryClan() const = 0;
 
 	DECLARE_EVENT(IClanInfo, FChangedEvent)
 	virtual FChangedEvent& OnChanged() = 0;
diff --git a/Engine/Source/Developer/FriendsAndChat/Private/Layers/Domain/ClanMemberList.cpp b/Engine/Source/Developer/FriendsAndChat/Private/Layers/Domain/ClanMemberList.cpp
--- a/Engine/Source/Developer/FriendsAndChat/Private/Layers/Domain/ClanMemberList.cpp
+++ b/Engine/Source/Developer/FriendsAndChat/Private/Layers/Domain/ClanMemberList.cpp
@@ -67,26 +67,24 @@ private:
 		OnFriendsListUpdated().Broadcast();
 	}
 
-	FClanMemberListImpl(TSharedRef<IClanInfo> InClanInfo, const TSharedRef<IFriendViewModelFactory>& InFriendViewModelFactory, const TSharedRef<FFriendsAndChatManager>& InFriendsAndChatManager)
+	FClanMemberListImpl(TSharedRef<IClanInfo> InClanInfo, const TSharedRef<IFriendViewModelFactory>& InFriendViewModelFactory)
 		: ClanInfo(InClanInfo)
-		, FriendsAndChatManager(InFriendsAndChatManager)
 		, FriendViewModelFactory(InFriendViewModelFactory)
 	{}
 
 private:
 
 	TSharedRef<IClanInfo> ClanInfo;
-	TWeakPtr<FFriendsAndChatManager> FriendsAndChatManager;
 	TSharedRef<IFriendViewModelFactory> FriendViewModelFactory;
 
 	friend FClanMemberListFactory;
 };
 
 TSharedRef< FClanMemberList > FClanMemberListFactory::Create(TSharedRef<IClanInfo> ClanInfo,
-	const TSharedRef<IFriendViewModelFactory>& FriendViewModelFactory,
-	const TSharedRef<FFriendsAndChatManager>& FriendsAndChatManager)
+	const TSharedRef<IFriendViewModelFactory>& FriendViewModelFactory
+	)
 {
-	TSharedRef< FClanMemberListImpl > ChatList(new FClanMemberListImpl(ClanInfo, FriendViewModelFactory, FriendsAndChatManager));
+	TSharedRef< FClanMemberListImpl > ChatList(new FClanMemberListImpl(ClanInfo, FriendViewModelFactory));
 	ChatList->Initialize();
 	return ChatList;
 }
\ No newline at end of file
diff --git a/Engine/Source/Developer/FriendsAndChat/Private/Layers/Domain/ClanMemberList.h b/Engine/Source/Developer/FriendsAndChat/Private/Layers/Domain/ClanMemberList.h
--- a/Engine/Source/Developer/FriendsAndChat/Private/Layers/Domain/ClanMemberList.h
+++ b/Engine/Source/Developer/FriendsAndChat/Private/Layers/Domain/ClanMemberList.h
@@ -20,5 +20,4 @@ class FClanMemberList
  */
 FACTORY(TSharedRef< FClanMemberList >, FClanMemberList,
 	TSharedRef<class IClanInfo> ClanInfo,
-	const TSharedRef<class IFriendViewModelFactory>& FriendViewModelFactory,
-	const TSharedRef<class FFriendsAndChatManager>& FriendsAndChatManager);
+	const TSharedRef<class IFriendViewModelFactory>& FriendViewModelFactory);
diff --git a/Engine/Source/Developer/FriendsAndChat/Private/Layers/Domain/ClanRepository.cpp b/Engine/Source/Developer/FriendsAndChat/Private/Layers/Domain/ClanRepository.cpp
--- a/Engine/Source/Developer/FriendsAndChat/Private/Layers/Domain/ClanRepository.cpp
+++ b/Engine/Source/Developer/FriendsAndChat/Private/Layers/Domain/ClanRepository.cpp
@@ -60,6 +60,11 @@ public:
 		return TEXT("Test Client");
 	}
 
+	FDateTime GetLastSeen() const
+	{
+		return FDateTime::UtcNow();
+	}
+
 	virtual const bool IsOnline() const override
 	{
 		return false;
@@ -182,6 +187,7 @@ public:
 			Member.MemberId = OnlineIdentity->CreateUniquePlayerId("Test4");
 			MemberList.Add(MakeShareable(new FClanMember(Member)));
 		}
+		bPrimaryClan = true;
 	}
 
 	virtual FText GetTitle() const override
@@ -189,6 +195,11 @@ public:
 		return FText::FromString("Test Title");
 	}
 
+	virtual const TSharedRef<const FUniqueNetId> GetUniqueID() const override
+	{
+		return UniqueID.ToSharedRef();
+	}
+
 	virtual FText GetDescription() const override
 	{
 		return FText::FromString("Test Description");
@@ -214,11 +225,18 @@ public:
 		return MemberList.Num();
 	}
 
+	virtual bool IsPrimaryClan() const override
+	{
+		return bPrimaryClan;
+	}
+
 	DECLARE_DERIVED_EVENT(FClanInfo, IClanInfo::FChangedEvent, FChangedEvent)
 	virtual IClanInfo::FChangedEvent& OnChanged() override { return ChangedEvent; }
 
 private:
 	TArray<TSharedPtr<IFriendItem> > MemberList;
+	TSharedPtr<FUniqueNetId> UniqueID;
+	bool bPrimaryClan;
 
 	FChangedEvent ChangedEvent;
 };
diff --git a/Engine/Source/Developer/FriendsAndChat/Private/Layers/Domain/DefaultFriendList.h b/Engine/Source/Developer/FriendsAndChat/Private/Layers/Domain/DefaultFriendList.h
--- a/Engine/Source/Developer/FriendsAndChat/Private/Layers/Domain/DefaultFriendList.h
+++ b/Engine/Source/Developer/FriendsAndChat/Private/Layers/Domain/DefaultFriendList.h
@@ -21,4 +21,4 @@ class FDefaultFriendList
 FACTORY(TSharedRef< FDefaultFriendList >, FDefaultFriendList,
 	EFriendsDisplayLists::Type ListType,
 	const TSharedRef<class IFriendViewModelFactory>& FriendViewModelFactory,
-	const TSharedRef<class FFriendsAndChatManager> & FriendAndChatManager);
+	const TSharedRef<class FFriendsService> & FriendsService);
diff --git a/Engine/Source/Developer/FriendsAndChat/Private/Layers/Domain/DefaultFriendsList.cpp b/Engine/Source/Developer/FriendsAndChat/Private/Layers/Domain/DefaultFriendsList.cpp
--- a/Engine/Source/Developer/FriendsAndChat/Private/Layers/Domain/DefaultFriendsList.cpp
+++ b/Engine/Source/Developer/FriendsAndChat/Private/Layers/Domain/DefaultFriendsList.cpp
@@ -4,6 +4,7 @@
 #include "DefaultFriendList.h"
 #include "IFriendItem.h"
 #include "FriendViewModel.h"
+#include "FriendsService.h"
 
 class FDefaultFriendListImpl
 	: public FDefaultFriendList
@@ -13,7 +14,7 @@ public:
 	virtual int32 GetFriendList(TArray< TSharedPtr<FFriendViewModel> >& OutFriendsList) override
 	{
 		TArray< TSharedPtr< IFriendItem > > FriendItemList;
-		FriendsAndChatManager.Pin()->GetFilteredFriendsList(FriendItemList);
+		FriendsService.Pin()->GetFilteredFriendsList(FriendItemList);
 		TArray< TSharedPtr< IFriendItem > > OnlineFriendsList;
 		TArray< TSharedPtr< IFriendItem > > OfflineFriendsList;
 
@@ -91,7 +92,7 @@ public:
 private:
 	void Initialize()
 	{
-		FriendsAndChatManager.Pin()->OnFriendsListUpdated().AddSP(this, &FDefaultFriendListImpl::HandleChatListUpdated);
+		FriendsService.Pin()->OnFriendsListUpdated().AddSP(this, &FDefaultFriendListImpl::HandleChatListUpdated);
 	}
 
 	void HandleChatListUpdated()
@@ -102,16 +103,16 @@ private:
 	FDefaultFriendListImpl(
 		EFriendsDisplayLists::Type InListType,
 		const TSharedRef<IFriendViewModelFactory>& InFriendViewModelFactory,
-		const TSharedRef<FFriendsAndChatManager>& InFriendsAndChatManager)
+		const TSharedRef<FFriendsService>& InFriendsService)
 		: ListType(InListType)
 		, FriendViewModelFactory(InFriendViewModelFactory)
-		, FriendsAndChatManager(InFriendsAndChatManager)
+		, FriendsService(InFriendsService)
 	{}
 
 private:
 	const EFriendsDisplayLists::Type ListType;
 	TSharedRef<IFriendViewModelFactory> FriendViewModelFactory;
-	TWeakPtr<FFriendsAndChatManager> FriendsAndChatManager;
+	TWeakPtr<FFriendsService> FriendsService;
 
 	friend FDefaultFriendListFactory;
 };
@@ -119,9 +120,9 @@ private:
 TSharedRef< FDefaultFriendList > FDefaultFriendListFactory::Create(
 	EFriendsDisplayLists::Type ListType,
 	const TSharedRef<IFriendViewModelFactory>& FriendViewModelFactory,
-	const TSharedRef<FFriendsAndChatManager>& FriendsAndChatManager)
+	const TSharedRef<FFriendsService>& FriendsService)
 {
-	TSharedRef< FDefaultFriendListImpl > ChatList(new FDefaultFriendListImpl(ListType, FriendViewModelFactory, FriendsAndChatManager));
+	TSharedRef< FDefaultFriendListImpl > ChatList(new FDefaultFriendListImpl(ListType, FriendViewModelFactory, FriendsService));
 	ChatList->Initialize();
 	return ChatList;
 }
\ No newline at end of file
diff --git a/Engine/Source/Developer/FriendsAndChat/Private/Layers/Domain/FriendsChatMarkupService.cpp b/Engine/Source/Developer/FriendsAndChat/Private/Layers/Domain/FriendsChatMarkupService.cpp
--- a/Engine/Source/Developer/FriendsAndChat/Private/Layers/Domain/FriendsChatMarkupService.cpp
+++ b/Engine/Source/Developer/FriendsAndChat/Private/Layers/Domain/FriendsChatMarkupService.cpp
@@ -162,15 +162,6 @@ class FFriendsChatMarkupServiceImpl
 {
 public:
 
-	virtual void ToggleDisplayChatTips() override
-	{
-		ForceDisplayToolTips = !ForceDisplayToolTips;
-		if(ForceDisplayToolTips)
-		{
-			GenerateDefaultTips();
-		}
-	}
-
 	virtual void CloseChatTips() override
 	{
 		ForceDisplayToolTips = false;
@@ -205,7 +196,7 @@ public:
 					}
 					else
 					{
-						CommunicationService->SendPrivateMessage(ProcessedInput->ValidFriends[0]->GetFriendItem()->GetUniqueID(), FText::FromString(ProcessedInput->Message));
+						CommunicationService->SendPrivateMessage(ProcessedInput->ValidFriends[0]->GetFriendItem(), FText::FromString(ProcessedInput->Message));
 						NavigationService->ChangeViewChannel(EChatMessageType::Whisper);
 					}
 				}
diff --git a/Engine/Source/Developer/FriendsAndChat/Private/Layers/Domain/FriendsChatMarkupService.h b/Engine/Source/Developer/FriendsAndChat/Private/Layers/Domain/FriendsChatMarkupService.h
--- a/Engine/Source/Developer/FriendsAndChat/Private/Layers/Domain/FriendsChatMarkupService.h
+++ b/Engine/Source/Developer/FriendsAndChat/Private/Layers/Domain/FriendsChatMarkupService.h
@@ -18,11 +18,6 @@ public:
 
 	virtual ~FFriendsChatMarkupService() {}
 
-	/**
-	 * Toggle the chat tips visibility.
-	 */
-	virtual void ToggleDisplayChatTips() = 0;
-
 	/**
 	 * Make chat tips unavailable.
 	 */
diff --git a/Engine/Source/Developer/FriendsAndChat/Private/Layers/Domain/FriendsListProvider.cpp b/Engine/Source/Developer/FriendsAndChat/Private/Layers/Domain/FriendsListProvider.cpp
--- a/Engine/Source/Developer/FriendsAndChat/Private/Layers/Domain/FriendsListProvider.cpp
+++ b/Engine/Source/Developer/FriendsAndChat/Private/Layers/Domain/FriendsListProvider.cpp
@@ -4,7 +4,7 @@
 #include "IFriendList.h"
 #include "DefaultFriendList.h"
 #include "RecentPlayerList.h"
-#include "FriendInviteList.h"
+#include "GameInviteList.h"
 #include "ClanMemberList.h"
 
 class FFriendListFactoryFactory;
@@ -19,21 +19,21 @@ public:
 		{
 			case EFriendsDisplayLists::GameInviteDisplay :
 			{
-				return FFriendInviteListFactory::Create(FriendViewModelFactory, FriendsAndChatManager.Pin().ToSharedRef());
+				return FGameInviteListFactory::Create(FriendViewModelFactory, GameAndPartyService.Pin().ToSharedRef());
 			}
 			break;
 			case EFriendsDisplayLists::RecentPlayersDisplay :
 			{
-				return FRecentPlayerListFactory::Create(FriendViewModelFactory, FriendsAndChatManager.Pin().ToSharedRef());
+				return FRecentPlayerListFactory::Create(FriendViewModelFactory, FriendsService.Pin().ToSharedRef());
 			}
 			break;
 		}
-		return FDefaultFriendListFactory::Create(ListType, FriendViewModelFactory, FriendsAndChatManager.Pin().ToSharedRef());
+		return FDefaultFriendListFactory::Create(ListType, FriendViewModelFactory, FriendsService.Pin().ToSharedRef());
 	}
 
 	virtual TSharedRef<IFriendList> Create(TSharedRef<IClanInfo> ClanInfo) override
 	{
-		return FClanMemberListFactory::Create(ClanInfo, FriendViewModelFactory, FriendsAndChatManager.Pin().ToSharedRef());
+		return FClanMemberListFactory::Create(ClanInfo, FriendViewModelFactory);
 	}
 
 	virtual ~FFriendListFactory() {}
@@ -42,24 +42,28 @@ private:
 
 	FFriendListFactory(
 	const TSharedRef<IFriendViewModelFactory>& InFriendViewModelFactory,
-	const TSharedRef<FFriendsAndChatManager>& InFriendsAndChatManager)
+	const TSharedRef<FFriendsService>& InFriendsService,
+	const TSharedRef<FGameAndPartyService>& InGameAndPartyService)
 		: FriendViewModelFactory(InFriendViewModelFactory)
-		, FriendsAndChatManager(InFriendsAndChatManager)
+		, FriendsService(InFriendsService)
+		, GameAndPartyService(InGameAndPartyService)
 	{ }
 
 private:
 
 	const TSharedRef<IFriendViewModelFactory> FriendViewModelFactory;
-	TWeakPtr<FFriendsAndChatManager> FriendsAndChatManager;
+	TWeakPtr<FFriendsService> FriendsService;
+	TWeakPtr<FGameAndPartyService> GameAndPartyService;
 	friend FFriendListFactoryFactory;
 };
 
 TSharedRef<IFriendListFactory> FFriendListFactoryFactory::Create(
 	const TSharedRef<IFriendViewModelFactory>& FriendViewModelFactory,
-	const TSharedRef<FFriendsAndChatManager>& FriendsAndChatManager)
+	const TSharedRef<FFriendsService>& FriendsService,
+	const TSharedRef<FGameAndPartyService>& GameAndPartyService)
 {
 	TSharedRef<FFriendListFactory> Factory = MakeShareable(
-		new FFriendListFactory(FriendViewModelFactory, FriendsAndChatManager));
+		new FFriendListFactory(FriendViewModelFactory, FriendsService, GameAndPartyService));
 
 	return Factory;
 }
\ No newline at end of file
