[INTEGRATE] merge UE4-LauncherDev > UE4 Friends and Chat up to CL 2633708

diff --git a/Engine/Source/Developer/FriendsAndChat/Private/Layers/Presentation/TimeStampRun.cpp b/Engine/Source/Developer/FriendsAndChat/Private/Layers/Presentation/TimeStampRun.cpp
new file mode 100644
--- /dev/null
+++ b/Engine/Source/Developer/FriendsAndChat/Private/Layers/Presentation/TimeStampRun.cpp
@@ -0,0 +1,232 @@
+// Copyright 1998-2015 Epic Games, Inc. All Rights Reserved.
+
+#include "FriendsAndChatPrivatePCH.h"
+#include "TimeStampRun.h"
+
+#include "SlateWidgetRun.h"
+
+TSharedRef< FTimeStampRun > FTimeStampRun::Create(const TSharedRef<class FTextLayout>& TextLayout, const FRunInfo& InRunInfo, const TSharedRef< const FString >& InText, const FTimeStampRunInfo& InWidgetInfo)
+{
+	TSharedRef<FTimeStampRun> TimeStampRun = MakeShareable(new FTimeStampRun(TextLayout, InRunInfo, InText, InWidgetInfo));
+	TimeStampRun->Initialize();
+	return TimeStampRun;
+}
+
+FTextRange FTimeStampRun::GetTextRange() const 
+{
+	return Range;
+}
+
+void FTimeStampRun::SetTextRange( const FTextRange& Value )
+{
+	Range = Value;
+}
+
+int16 FTimeStampRun::GetBaseLine( float Scale ) const 
+{
+	return Info.Baseline * Scale;
+}
+
+int16 FTimeStampRun::GetMaxHeight( float Scale ) const 
+{
+	return Info.Size.Get( TimeStampWidget->GetDesiredSize() ).Y * Scale;
+}
+
+FVector2D FTimeStampRun::Measure( int32 StartIndex, int32 EndIndex, float Scale ) const 
+{
+	return Info.Size.Get( TimeStampWidget->GetDesiredSize() ) * Scale;
+}
+
+int8 FTimeStampRun::GetKerning( int32 CurrentIndex, float Scale ) const 
+{
+	return 0;
+}
+
+TSharedRef< ILayoutBlock > FTimeStampRun::CreateBlock( int32 StartIndex, int32 EndIndex, FVector2D Size, const TSharedPtr< IRunRenderer >& Renderer )
+{
+	return FDefaultLayoutBlock::Create( SharedThis( this ), FTextRange( StartIndex, EndIndex ), Size, Renderer );
+}
+
+int32 FTimeStampRun::OnPaint( const FPaintArgs& Args, const FTextLayout::FLineView& Line, const TSharedRef< ILayoutBlock >& Block, const FTextBlockStyle& DefaultStyle, const FGeometry& AllottedGeometry, const FSlateRect& MyClippingRect, FSlateWindowElementList& OutDrawElements, int32 LayerId, const FWidgetStyle& InWidgetStyle, bool bParentEnabled ) const 
+{
+	// The block size and offset values are pre-scaled, so we need to account for that when converting the block offsets into paint geometry
+	const float InverseScale = Inverse(AllottedGeometry.Scale);
+
+	const FVector2D DesiredWidgetSize = TimeStampWidget->GetDesiredSize();
+	if ( (DesiredWidgetSize != WidgetSize ) || ForceRedraw)
+	{
+		WidgetSize = DesiredWidgetSize;
+
+		const TSharedPtr<FTextLayout> TextLayoutPtr = TextLayout.Pin();
+		if (TextLayoutPtr.IsValid())
+		{
+			TextLayoutPtr->DirtyRunLayout(SharedThis(this));
+		}
+		ForceRedraw = false;
+	}
+
+	const FGeometry WidgetGeometry = AllottedGeometry.MakeChild(TransformVector(InverseScale, Block->GetSize()), FSlateLayoutTransform(TransformPoint(InverseScale, Block->GetLocationOffset())));
+	return TimeStampWidget->Paint( Args, WidgetGeometry, MyClippingRect, OutDrawElements, LayerId, InWidgetStyle, bParentEnabled );
+}
+
+const TArray< TSharedRef<SWidget> >& FTimeStampRun::GetChildren()
+{
+	return Children;
+}
+
+void FTimeStampRun::ArrangeChildren( const TSharedRef< ILayoutBlock >& Block, const FGeometry& AllottedGeometry, FArrangedChildren& ArrangedChildren ) const
+{
+	// The block size and offset values are pre-scaled, so we need to account for that when converting the block offsets into paint geometry
+	const float InverseScale = Inverse(AllottedGeometry.Scale);
+
+	ArrangedChildren.AddWidget(
+		AllottedGeometry.MakeChild(TimeStampWidget.ToSharedRef(), TransformVector(InverseScale, Block->GetSize()), FSlateLayoutTransform(TransformPoint(InverseScale, Block->GetLocationOffset())))
+		);
+}
+
+int32 FTimeStampRun::GetTextIndexAt( const TSharedRef< ILayoutBlock >& Block, const FVector2D& Location, float Scale, ETextHitPoint* const OutHitPoint ) const
+{
+	// A widget should always contain a single character (a breaking space)
+	check(Range.Len() == 1);
+
+	const FVector2D& BlockOffset = Block->GetLocationOffset();
+	const FVector2D& BlockSize = Block->GetSize();
+
+	const float Left = BlockOffset.X;
+	const float Top = BlockOffset.Y;
+	const float Right = BlockOffset.X + BlockSize.X;
+	const float Bottom = BlockOffset.Y + BlockSize.Y;
+
+	const bool ContainsPoint = Location.X >= Left && Location.X < Right && Location.Y >= Top && Location.Y < Bottom;
+
+	if ( !ContainsPoint )
+	{
+		return INDEX_NONE;
+	}
+
+	const FVector2D ScaledWidgetSize = TimeStampWidget->GetDesiredSize() * Scale;
+	const int32 Index = (Location.X <= (Left + (ScaledWidgetSize.X * 0.5f))) ? Range.BeginIndex : Range.EndIndex;
+	
+	if (OutHitPoint)
+	{
+		const FTextRange BlockRange = Block->GetTextRange();
+		*OutHitPoint = (Index == BlockRange.EndIndex) ? ETextHitPoint::RightGutter : ETextHitPoint::WithinText;
+	}
+	return Index;
+}
+
+FVector2D FTimeStampRun::GetLocationAt( const TSharedRef< ILayoutBlock >& Block, int32 Offset, float Scale ) const
+{
+	return Block->GetLocationOffset();
+}
+
+void FTimeStampRun::Move(const TSharedRef<FString>& NewText, const FTextRange& NewRange)
+{
+	Text = NewText;
+	Range = NewRange;
+}
+
+TSharedRef<IRun> FTimeStampRun::Clone() const
+{
+	return FTimeStampRun::Create(TextLayout.Pin().ToSharedRef(), RunInfo, Text, Info);
+}
+
+void FTimeStampRun::AppendTextTo(FString& AppendToText) const
+{
+	AppendToText.Append(**Text + Range.BeginIndex, Range.Len());
+}
+
+void FTimeStampRun::AppendTextTo(FString& AppendToText, const FTextRange& PartialRange) const
+{
+	check(Range.BeginIndex <= PartialRange.BeginIndex);
+	check(Range.EndIndex >= PartialRange.EndIndex);
+
+	AppendToText.Append(**Text + PartialRange.BeginIndex, PartialRange.Len());
+}
+
+const FRunInfo& FTimeStampRun::GetRunInfo() const
+{
+	return RunInfo;
+}
+
+ERunAttributes FTimeStampRun::GetRunAttributes() const
+{
+	return ERunAttributes::None;
+}
+
+bool FTimeStampRun::UpdateTimeStamp(float Delay)
+{
+	bool ContinueTicking = true;
+
+	FDateTime Now = FDateTime::Now();
+
+	if ((Now - TimeStamp).GetDuration() <= FTimespan::FromMinutes(1.0))
+	{
+		static FText NowTimeStamp = NSLOCTEXT("FTimeStampRun", "Now_TimeStamp", "Now");
+		MessageTimeAsText = NowTimeStamp;
+	}
+	else if ((Now - TimeStamp).GetDuration() <= FTimespan::FromMinutes(2.0))
+	{
+		static FText OneMinuteTimeStamp = NSLOCTEXT("FTimeStampRun", "1_Minute_TimeStamp", "1 min");
+		MessageTimeAsText = OneMinuteTimeStamp;
+	}
+	else if ((Now - TimeStamp).GetDuration() <= FTimespan::FromMinutes(3.0))
+	{
+		static FText TwoMinuteTimeStamp = NSLOCTEXT("FTimeStampRun", "2_Minute_TimeStamp", "2 min");
+		MessageTimeAsText = TwoMinuteTimeStamp;
+	}
+	else if ((Now - TimeStamp).GetDuration() <= FTimespan::FromMinutes(4.0))
+	{
+		static FText ThreeMinuteTimeStamp = NSLOCTEXT("FTimeStampRun", "3_Minute_TimeStamp", "3 min");
+		MessageTimeAsText = ThreeMinuteTimeStamp;
+	}
+	else if ((Now - TimeStamp).GetDuration() <= FTimespan::FromMinutes(5.0))
+	{
+		static FText FourMinuteTimeStamp = NSLOCTEXT("FTimeStampRun", "4_Minute_TimeStamp", "4 min");
+		MessageTimeAsText = FourMinuteTimeStamp;
+	}
+	else if ((Now - TimeStamp).GetDuration() <= FTimespan::FromMinutes(6.0))
+	{
+		static FText FiveMinuteTimeStamp = NSLOCTEXT("FTimeStampRun", "5_Minute_TimeStamp", "5 min");
+		MessageTimeAsText = FiveMinuteTimeStamp;
+	}
+	else
+	{
+		MessageTimeAsText = FText::AsTime(TimeStamp, EDateTimeStyle::Short, "GMT");
+		ContinueTicking = false;
+	}
+
+	ForceRedraw = true;
+
+	return ContinueTicking;
+}
+
+void FTimeStampRun::Initialize()
+{
+	TimeStampWidget = SNew(STextBlock)
+		.TextStyle(&Info.TextBoxStyle)
+		.Text(this, &FTimeStampRun::GetMessageTimeText);
+
+	if (UpdateTimeStamp(60)) //this will broadcast a changed event
+	{
+		TickerHandle = FTicker::GetCoreTicker().AddTicker(FTickerDelegate::CreateSP(this, &FTimeStampRun::UpdateTimeStamp), 60);
+	}
+
+	TimeStampWidget->SlatePrepass();
+	WidgetSize = TimeStampWidget->GetDesiredSize();
+	Children.Add(TimeStampWidget.ToSharedRef());
+}
+
+
+FTimeStampRun::FTimeStampRun(const TSharedRef<class FTextLayout>& InTextLayout, const FRunInfo& InRunInfo, const TSharedRef< const FString >& InText, const FTimeStampRunInfo& InWidgetInfo)
+	: TextLayout(InTextLayout)
+	, RunInfo( InRunInfo )
+	, Text( InText )
+	, Range( 0, 1 )
+	, Info( InWidgetInfo )
+	, Children()
+	, TimeStamp(InWidgetInfo.TimeStamp)
+	, ForceRedraw(false)
+	, WidgetSize()
+{
+}
diff --git a/Engine/Source/Developer/FriendsAndChat/Private/Layers/Presentation/TimeStampRun.h b/Engine/Source/Developer/FriendsAndChat/Private/Layers/Presentation/TimeStampRun.h
new file mode 100644
--- /dev/null
+++ b/Engine/Source/Developer/FriendsAndChat/Private/Layers/Presentation/TimeStampRun.h
@@ -0,0 +1,82 @@
+// Copyright 1998-2015 Epic Games, Inc. All Rights Reserved.
+#pragma once
+#include "ISlateRun.h"
+
+class FTimeStampRun : public ISlateRun, public TSharedFromThis< FTimeStampRun >
+{
+public:
+
+	struct FTimeStampRunInfo
+	{
+		FDateTime TimeStamp;
+		int16 Baseline;
+		FTextBlockStyle TextBoxStyle;
+		TOptional< FVector2D > Size;
+
+		FTimeStampRunInfo(FDateTime InTimeStamp, int16 InBaseline, const FTextBlockStyle& InTextBoxStyle, TOptional< FVector2D > InSize = TOptional< FVector2D >())
+			: TimeStamp(InTimeStamp)
+			, Baseline(InBaseline)
+			, TextBoxStyle(InTextBoxStyle)
+			, Size(InSize)
+		{
+
+		}
+	};
+
+	static TSharedRef< FTimeStampRun > Create(const TSharedRef<class FTextLayout>& TextLayout, const FRunInfo& InRunInfo, const TSharedRef< const FString >& InText, const FTimeStampRunInfo& InWidgetInfo);
+
+	virtual ~FTimeStampRun() {}
+
+	virtual FTextRange GetTextRange() const override;
+	virtual void SetTextRange( const FTextRange& Value ) override;
+	virtual int16 GetBaseLine( float Scale ) const override;
+	virtual int16 GetMaxHeight( float Scale ) const override;
+	virtual FVector2D Measure( int32 StartIndex, int32 EndIndex, float Scale ) const override;
+	virtual int8 GetKerning(int32 CurrentIndex, float Scale) const override;
+	virtual TSharedRef< ILayoutBlock > CreateBlock( int32 StartIndex, int32 EndIndex, FVector2D Size, const TSharedPtr< IRunRenderer >& Renderer ) override;
+	virtual int32 OnPaint( const FPaintArgs& Args, const FTextLayout::FLineView& Line, const TSharedRef< ILayoutBlock >& Block, const FTextBlockStyle& DefaultStyle, const FGeometry& AllottedGeometry, const FSlateRect& MyClippingRect, FSlateWindowElementList& OutDrawElements, int32 LayerId, const FWidgetStyle& InWidgetStyle, bool bParentEnabled ) const override;
+	virtual const TArray< TSharedRef<SWidget> >& GetChildren() override;
+	virtual void ArrangeChildren( const TSharedRef< ILayoutBlock >& Block, const FGeometry& AllottedGeometry, FArrangedChildren& ArrangedChildren ) const override;
+	virtual int32 GetTextIndexAt( const TSharedRef< ILayoutBlock >& Block, const FVector2D& Location, float Scale, ETextHitPoint* const OutHitPoint = nullptr ) const override;
+	virtual FVector2D GetLocationAt( const TSharedRef< ILayoutBlock >& Block, int32 Offset, float Scale ) const override;
+	virtual void BeginLayout() override {}
+	virtual void EndLayout() override {}
+	virtual void Move(const TSharedRef<FString>& NewText, const FTextRange& NewRange) override;
+	virtual TSharedRef<IRun> Clone() const override;
+	virtual void AppendTextTo(FString& Text) const override;
+	virtual void AppendTextTo(FString& AppendToText, const FTextRange& PartialRange) const override;
+	virtual const FRunInfo& GetRunInfo() const override;
+	virtual ERunAttributes GetRunAttributes() const override;
+
+private:
+
+	void Initialize();
+
+	FText GetMessageTimeText() const
+	{
+		return MessageTimeAsText;
+	}
+
+	bool UpdateTimeStamp(float Delay);
+
+	FTimeStampRun(const TSharedRef<class FTextLayout>& TextLayout, const FRunInfo& InRunInfo, const TSharedRef< const FString >& InText, const FTimeStampRunInfo& InWidgetInfo);
+
+	FTimeStampRun( const FTimeStampRun& Run );
+
+private:
+
+	TWeakPtr<class FTextLayout> TextLayout;
+	FRunInfo RunInfo;
+	TSharedRef< const FString > Text;
+	FTextRange Range;
+	FTimeStampRunInfo Info;
+	TArray< TSharedRef<SWidget> > Children;
+	TSharedPtr<SWidget> TimeStampWidget;
+	FText MessageTimeAsText;
+	FDelegateHandle TickerHandle;
+	FDateTime TimeStamp;
+
+	mutable bool ForceRedraw;
+	mutable FVector2D WidgetSize;
+
+};
diff --git a/Engine/Source/Developer/FriendsAndChat/Private/Layers/UI/SAutoTextScroller.cpp b/Engine/Source/Developer/FriendsAndChat/Private/Layers/UI/SAutoTextScroller.cpp
--- a/Engine/Source/Developer/FriendsAndChat/Private/Layers/UI/SAutoTextScroller.cpp
+++ b/Engine/Source/Developer/FriendsAndChat/Private/Layers/UI/SAutoTextScroller.cpp
@@ -7,9 +7,13 @@ namespace EScrollerState
 {
 	enum Type
 	{
-		Start = 0,
+		FadeOn = 0,
+		Start,
+		Start_Wait,
 		Scrolling,
-		End
+		Stop,
+		Stop_Wait,
+		FadeOff
 	};
 }
 
@@ -32,36 +36,72 @@ public:
 			[
 				SNew(STextBlock)
 				.Font(FontStyle.FriendsFontSmallBold)
-				.ColorAndOpacity(FontStyle.DefaultDullFontColor)
+				.ColorAndOpacity(this, &SAutoTextScrollerImpl::GetFontColor)
 				.Text(InArgs._Text)
+				//.Text(NSLOCTEXT("ughu", "eref", "This is a very long sentence to test the auto scroller 5 4 3 2 1 ."))
 			]
 		]);
 
 		TimeElapsed = 0.f;
 		ScrollOffset = 0;
+		FontAlpha = 1.0f;
 		ScrollerState = EScrollerState::Start;
 	}
 
+	FSlateColor GetFontColor() const
+	{
+		FLinearColor RetVal = FontStyle.DefaultDullFontColor;
+		RetVal.A = FontAlpha;
+		return RetVal;
+	}
+
 	virtual void Tick(const FGeometry& AllottedGeometry, const double InCurrentTime, const float InDeltaTime) override
 	{
 		const float ContentSize = ScrollBox->GetDesiredSize().X;
 		const float Speed = 20.0f;
 		const float StartDelay = 2.0f;
 		const float EndDelay = 2.0f;
+		const float FadeOnDelay = 0.5f;
+		const float FadeOffDelay = 0.5f;
 
 		TimeElapsed += InDeltaTime;
 
 		switch (ScrollerState)
 		{
+			case EScrollerState::FadeOn:
+			{
+				FontAlpha = TimeElapsed / FadeOnDelay;
+				if (TimeElapsed >= FadeOnDelay)
+				{
+					FontAlpha = 1.0f;
+					TimeElapsed = 0.0f;
+					ScrollOffset = 0.0f;
+					ScrollerState = EScrollerState::Start;
+				}
+			}break;
 			case EScrollerState::Start:
 			{
+				// Unregister Tick Delegate
 				auto PinnedActiveTimerHandle = ActiveTimerHandle.Pin();
 				if (PinnedActiveTimerHandle.IsValid())
 				{
 					UnRegisterActiveTimer(PinnedActiveTimerHandle.ToSharedRef());
 				}
-				ActiveTimerHandle = RegisterActiveTimer(StartDelay, FWidgetActiveTimerDelegate::CreateSP(this, &SAutoTextScrollerImpl::UpdateScrollerState));
+
+				// Check to see if we need to scroll
+				if (ExternalScrollbar->DistanceFromBottom() == 0.0f && ExternalScrollbar->DistanceFromTop() == 0.0f)
+				{
+					// Don't run auto scrolling if text already fits.
+					break;
+				}
+				else
+				{
+					ActiveTimerHandle = RegisterActiveTimer(StartDelay, FWidgetActiveTimerDelegate::CreateSP(this, &SAutoTextScrollerImpl::UpdateScrollerState));
+					ScrollerState = EScrollerState::Start_Wait;
+				}
 			}break;
+			case EScrollerState::Start_Wait:
+				break;
 			case EScrollerState::Scrolling:
 			{
 				ScrollOffset += Speed * InDeltaTime;
@@ -69,18 +109,43 @@ public:
 				{
 					ScrollOffset = ContentSize;
 					TimeElapsed = 0.0f;
-					ScrollerState = EScrollerState::End;
+					ScrollerState = EScrollerState::Stop;
 				}
 				
 			}break;
-			case EScrollerState::End:
+			case EScrollerState::Stop:
 			{
+				// Unregister Tick Delegate
 				auto PinnedActiveTimerHandle = ActiveTimerHandle.Pin();
 				if (PinnedActiveTimerHandle.IsValid())
 				{
 					UnRegisterActiveTimer(PinnedActiveTimerHandle.ToSharedRef());
 				}
-				ActiveTimerHandle = RegisterActiveTimer(EndDelay, FWidgetActiveTimerDelegate::CreateSP(this, &SAutoTextScrollerImpl::UpdateScrollerState));
+
+				// Check to see if we need to scroll
+				if (ExternalScrollbar->DistanceFromBottom() == 0.0f && ExternalScrollbar->DistanceFromTop() == 0.0f)
+				{
+					// Don't run auto scrolling if text already fits.
+					break;
+				}
+				else
+				{
+					ActiveTimerHandle = RegisterActiveTimer(EndDelay, FWidgetActiveTimerDelegate::CreateSP(this, &SAutoTextScrollerImpl::UpdateScrollerState));
+					ScrollerState = EScrollerState::Stop_Wait;
+				}
+			}break;
+			case EScrollerState::Stop_Wait:
+				break;
+			case EScrollerState::FadeOff:
+			{
+				FontAlpha = 1.0f - TimeElapsed / FadeOffDelay;
+				if (TimeElapsed >= FadeOffDelay)
+				{
+					FontAlpha = 0.0f;
+					TimeElapsed = 0.0f;
+					ScrollOffset = 0.0f;
+					ScrollerState = EScrollerState::FadeOn;
+				}
 			}break;
 		};
 		ScrollBox->SetScrollOffset(ScrollOffset);
@@ -88,19 +153,32 @@ public:
 
 	EActiveTimerReturnType UpdateScrollerState(double InCurrentTime, float InDeltaTime)
 	{
-		if (ScrollerState == EScrollerState::End)
+		if (ScrollerState == EScrollerState::Stop_Wait)
 		{
 			TimeElapsed = 0.0f;
-			ScrollOffset = 0.0f;
-			ScrollerState = EScrollerState::Start;
+			ScrollerState = EScrollerState::FadeOff;
 			ScrollBox->SetScrollOffset(ScrollOffset);
+
+			auto PinnedActiveTimerHandle = ActiveTimerHandle.Pin();
+			if (PinnedActiveTimerHandle.IsValid())
+			{
+				UnRegisterActiveTimer(PinnedActiveTimerHandle.ToSharedRef());
+			}
+			ActiveTimerHandle = RegisterActiveTimer(0, FWidgetActiveTimerDelegate::CreateSP(this, &SAutoTextScrollerImpl::UpdateScrollerState));
 		}
-		else if (ScrollerState == EScrollerState::Start)
+		else if (ScrollerState == EScrollerState::Start_Wait)
 		{
 			TimeElapsed = 0.0f;
 			ScrollOffset = 0.0f;
 			ScrollerState = EScrollerState::Scrolling;
 			ScrollBox->SetScrollOffset(ScrollOffset);
+
+			auto PinnedActiveTimerHandle = ActiveTimerHandle.Pin();
+			if (PinnedActiveTimerHandle.IsValid())
+			{
+				UnRegisterActiveTimer(PinnedActiveTimerHandle.ToSharedRef());
+			}
+			ActiveTimerHandle = RegisterActiveTimer(0, FWidgetActiveTimerDelegate::CreateSP(this, &SAutoTextScrollerImpl::UpdateScrollerState));
 		}
 		return EActiveTimerReturnType::Continue;
 	}
@@ -115,6 +193,7 @@ private:
 	EScrollerState::Type ScrollerState;
 	float TimeElapsed;
 	float ScrollOffset;
+	float FontAlpha;
 };
 
 TSharedRef<SAutoTextScroller> SAutoTextScroller::New()
diff --git a/Engine/Source/Developer/FriendsAndChat/Private/Layers/UI/SChatChrome.cpp b/Engine/Source/Developer/FriendsAndChat/Private/Layers/UI/SChatChrome.cpp
--- a/Engine/Source/Developer/FriendsAndChat/Private/Layers/UI/SChatChrome.cpp
+++ b/Engine/Source/Developer/FriendsAndChat/Private/Layers/UI/SChatChrome.cpp
@@ -24,13 +24,12 @@ public:
 
 		ChromeViewModel->OnActiveTabChanged().AddSP(this, &SChatChromeImpl::HandleActiveTabChanged);
 		ChromeViewModel->OnVisibleTabsChanged().AddSP(this, &SChatChromeImpl::HandleVisibleTabsChanged);
-		FLinearColor BorderColor = ChromeViewModel->IsFading() ? FLinearColor::Transparent : FLinearColor::White;
 
 		SUserWidget::Construct(SUserWidget::FArguments()
 		[
 			SNew(SBorder)
-			.BorderImage(&FriendStyle.FriendsChatChromeStyle.ChatBackgroundBrush)
-			.BorderBackgroundColor(BorderColor)
+			.Padding(0)
+			.BorderImage(FCoreStyle::Get().GetBrush("NoBorder"))
 			.Visibility(this, &SChatChromeImpl::GetChatWindowVisibility)
 			.Padding(0)
 			[
diff --git a/Engine/Source/Developer/FriendsAndChat/Private/Layers/UI/SChatEntryWidget.cpp b/Engine/Source/Developer/FriendsAndChat/Private/Layers/UI/SChatEntryWidget.cpp
--- a/Engine/Source/Developer/FriendsAndChat/Private/Layers/UI/SChatEntryWidget.cpp
+++ b/Engine/Source/Developer/FriendsAndChat/Private/Layers/UI/SChatEntryWidget.cpp
@@ -6,6 +6,7 @@
 #include "RichTextLayoutMarshaller.h"
 #include "ChatViewModel.h"
 #include "ChatEntryCommands.h"
+#include "SFriendActions.h"
 
 #define LOCTEXT_NAMESPACE ""
 
@@ -27,34 +28,19 @@ public:
 		float HightOverride = FriendStyle.FriendsChatStyle.ChatEntryHeight;
 		MaxChatLength = InArgs._MaxChatLength;
 		HintText = InArgs._HintText;
-		Marshaller = InArgs._Marshaller;
+
+		Marshaller = FRichTextLayoutMarshaller::Create(
+			TArray<TSharedRef<ITextDecorator>>(), 
+			&FFriendsAndChatModuleStyle::Get()
+			);
+
+		OnHyperlinkClicked = FSlateHyperlinkRun::FOnClick::CreateSP(SharedThis(this), &SChatEntryWidgetImpl::HandleNameClicked);
+
+		Marshaller->AppendInlineDecorator(FHyperlinkDecorator::Create(TEXT("UserName"), OnHyperlinkClicked));
 
 		SUserWidget::Construct(SUserWidget::FArguments()
 		[
-			SNew(SHorizontalBox)
-			+SHorizontalBox::Slot()
-			.AutoWidth()
-			[
-			 	SNew(SBox)
-				.Visibility(ViewModel->AllowMarkup() ? EVisibility::Visible : EVisibility::Collapsed)
-			 	.WidthOverride(HightOverride)
-			 	[
-					SNew(SButton)
-					.OnClicked(this, &SChatEntryWidgetImpl::HandleActionDropDownClicked)
-					.VAlign(VAlign_Center)
-					.HAlign(HAlign_Center)
-					.Cursor(EMouseCursor::Hand)
-					[
-						SNew(STextBlock)
-						.Font(FriendStyle.FriendsNormalFontStyle.FriendsFontLargeBold)
-						.Text(FText::FromString("/"))
-					]
-				 ]
-			]
-			+SHorizontalBox::Slot()
-			[
-				SAssignNew(ChatBox, SBox)
-			]
+			SAssignNew(ChatBox, SBox)
 		]);
 
 		RebuildTextEntry();
@@ -129,12 +115,6 @@ public:
 	// End SUserWidget
 private:
 
-	FReply HandleActionDropDownClicked()
-	{
-		ViewModel->ToggleChatTipVisibility();
-		return FReply::Handled();
-	}
-
 	/**
 	 * Handle chat commited - send the message to be broadcast from the message service
 	 * @param CommentText	The committed text
@@ -253,6 +233,27 @@ private:
 		return FMath::Max(WindowWidth - 60, 0.0f);
 	}
 
+	void HandleNameClicked(const FSlateHyperlinkRun::FMetadata& Metadata)
+	{
+		const FString* UsernameString = Metadata.Find(TEXT("Username"));
+		const FString* UniqueIDString = Metadata.Find(TEXT("uid"));
+
+		if (UsernameString && UniqueIDString)
+		{
+			FText Username = FText::FromString(*UsernameString);
+			const TSharedRef<FFriendViewModel> FriendViewModel = ViewModel->GetFriendViewModel(*UniqueIDString, Username).ToSharedRef();
+			TSharedRef<SWidget> Widget = SNew(SFriendActions, FriendViewModel).FriendStyle(&FriendStyle);
+
+			FSlateApplication::Get().PushMenu(
+				SharedThis(this),
+				FWidgetPath(),
+				Widget,
+				FSlateApplication::Get().GetCursorPos(),
+				FPopupTransitionEffect(FPopupTransitionEffect::ContextMenu)
+				);
+		}
+	}
+
 private:
 	/* Holds the max chat length */
 	int32 MaxChatLength;
@@ -269,7 +270,8 @@ private:
 	/* Holds the hint text for the chat box*/
 	TAttribute< FText > HintText;
 	/* Holds the marshaller for the chat box*/
-	TSharedPtr<ITextLayoutMarshaller> Marshaller;
+	TSharedPtr<FRichTextLayoutMarshaller> Marshaller;
+	FSlateHyperlinkRun::FOnClick OnHyperlinkClicked;
 };
 
 TSharedRef<SChatEntryWidget> SChatEntryWidget::New()
diff --git a/Engine/Source/Developer/FriendsAndChat/Private/Layers/UI/SChatEntryWidget.h b/Engine/Source/Developer/FriendsAndChat/Private/Layers/UI/SChatEntryWidget.h
--- a/Engine/Source/Developer/FriendsAndChat/Private/Layers/UI/SChatEntryWidget.h
+++ b/Engine/Source/Developer/FriendsAndChat/Private/Layers/UI/SChatEntryWidget.h
@@ -16,9 +16,6 @@ public:
 
 		SLATE_ARGUMENT(const FFriendsAndChatStyle*, FriendStyle)	
 
-		/** The marshaller used to get/set the raw text to/from the text layout. */
-		SLATE_ARGUMENT(TSharedPtr< ITextLayoutMarshaller >, Marshaller)
-
 		// The hint that shows what key activates chat
 		SLATE_ATTRIBUTE(FText, HintText)
 
diff --git a/Engine/Source/Developer/FriendsAndChat/Private/Layers/UI/SChatWindow.cpp b/Engine/Source/Developer/FriendsAndChat/Private/Layers/UI/SChatWindow.cpp
--- a/Engine/Source/Developer/FriendsAndChat/Private/Layers/UI/SChatWindow.cpp
+++ b/Engine/Source/Developer/FriendsAndChat/Private/Layers/UI/SChatWindow.cpp
@@ -10,7 +10,7 @@
 #include "SChatMarkupTips.h"
 #include "SFriendActions.h"
 #include "SChatEntryWidget.h"
-#include "RichTextLayoutMarshaller.h"
+#include "ChatTextLayoutMarshaller.h"
 #include "TextDecorators.h"
 
 #define LOCTEXT_NAMESPACE "SChatWindow"
@@ -31,19 +31,12 @@ public:
 
 		FFriendsAndChatModuleStyle::Initialize(FriendStyle);
 
-		RichTextMarshaller = FRichTextLayoutMarshaller::Create(
-			TArray<TSharedRef<ITextDecorator>>(), 
-			&FFriendsAndChatModuleStyle::Get()
-			);
-
 		OnHyperlinkClicked = FSlateHyperlinkRun::FOnClick::CreateSP(SharedThis(this), &SChatWindowImpl::HandleNameClicked);
 		OnChannelHyperlinkClicked = FSlateHyperlinkRun::FOnClick::CreateSP(SharedThis(this), &SChatWindowImpl::HandleChannelClicked);
-		OnTimeStampCreated = FWidgetDecorator::FCreateWidget::CreateSP(SharedThis(this), &SChatWindowImpl::HandleTimeStampDisplay);
 
-		RichTextMarshaller->AppendInlineDecorator(FHyperlinkDecorator::Create(TEXT("UserName"), OnHyperlinkClicked));
-		RichTextMarshaller->AppendInlineDecorator(FHyperlinkDecorator::Create(TEXT("Channel"), OnChannelHyperlinkClicked));
-		RichTextMarshaller->AppendInlineDecorator(FImageDecorator::Create(TEXT("img"), &FFriendsAndChatModuleStyle::Get()));
-		RichTextMarshaller->AppendInlineDecorator(FWidgetDecorator::Create(TEXT("TimeStamp"), OnTimeStampCreated));
+		RichTextMarshaller = FChatTextLayoutMarshallerFactory::Create(&FriendStyle, ViewModel->MultiChat());
+		RichTextMarshaller->AddChannelNameHyperlinkDecorator(FHyperlinkDecorator::Create(TEXT("Channel"), OnChannelHyperlinkClicked));
+		RichTextMarshaller->AddUserNameHyperlinkDecorator(FHyperlinkDecorator::Create(TEXT("UserName"), OnHyperlinkClicked));
 
 		TimeTransparency = 0.0f;
 
@@ -93,7 +86,7 @@ public:
 						]
 					]
 					+SOverlay::Slot()
-					.Padding(FMargin(40, 0))
+					.Padding(FMargin(10, 0))
 					.VAlign(VAlign_Bottom)
 					[
 						SNew(SBorder)
@@ -128,7 +121,6 @@ public:
 								[
 									SAssignNew(ChatTextBox, SChatEntryWidget, ViewModel.ToSharedRef())
 								 	.FriendStyle(&FriendStyle)
-									.Marshaller(RichTextMarshaller)
 									.HintText(InArgs._ActivationHintText)
 								]
 								+SVerticalBox::Slot()
@@ -249,7 +241,6 @@ private:
 			.Font(FFriendsAndChatModuleStyle::GetStyleService()->GetNormalFont())
 			.WrapTextAt(this, &SChatWindowImpl::GetChatWrapWidth)
 			.Marshaller(RichTextMarshaller.ToSharedRef())
-			.Text(this, &SChatWindowImpl::GetChatText)
 			.IsReadOnly(true);
 
 			ChatScrollBox->AddSlot()
@@ -268,6 +259,7 @@ private:
 			CreateChatWidgets();
 		}
 
+		// ToDo - Move the following logic into the view model
 		TArray<TSharedRef<FChatItemViewModel>> Messages = ViewModel->GetMessages();
 		for (;LastMessageIndex < Messages.Num(); ++LastMessageIndex)
 		{
@@ -282,20 +274,18 @@ private:
 				*LastMessage->GetRecipientID() == *ChatItem->GetRecipientID() &&
 				(ChatItem->GetMessageTime() - LastMessage->GetMessageTime()).GetDuration() <= FTimespan::FromMinutes(1.0))
 			{
-				ChatText.AppendLine(FormatListMessageText(ChatItem, true));
+				RichTextMarshaller->AppendMessage(ChatItem, true);
 			}
 			else
 			{
-				static FString MessageBreak(TEXT("<MessageBreak></>"));
-				ChatText.AppendLine(MessageBreak);
-				ChatText.AppendLine(FormatListMessageText(ChatItem, false));
+				RichTextMarshaller->AppendLineBreak();
+				RichTextMarshaller->AppendMessage(ChatItem, false);
 			}
 			LastMessage = ChatItem;
 		}
 
 		if (RichText.IsValid() && ChatScrollBox.IsValid())
 		{
-			RichText->Refresh();
 			if(ExternalScrollbar.IsValid() && ExternalScrollbar->DistanceFromBottom() == 0)
 			{
 				ChatScrollBox->ScrollToEnd();
@@ -303,11 +293,6 @@ private:
 		}
 	}
 
-	FText GetChatText() const
-	{
-		return ChatText.ToText();
-	}
-
 	void HandleNameClicked(const FSlateHyperlinkRun::FMetadata& Metadata)
 	{
 		const FString* UsernameString = Metadata.Find(TEXT("Username"));
@@ -317,7 +302,7 @@ private:
 		{
 			FText Username = FText::FromString(*UsernameString);
 			const TSharedRef<FFriendViewModel> FriendViewModel = ViewModel->GetFriendViewModel(*UniqueIDString, Username).ToSharedRef();
-			TSharedRef<SWidget> Widget = SNew(SFriendActions, FriendViewModel).FriendStyle(&FriendStyle);
+			TSharedRef<SWidget> Widget = SNew(SFriendActions, FriendViewModel).FriendStyle(&FriendStyle).FromChat(true);
 
 			FSlateApplication::Get().PushMenu(
 				SharedThis(this),
@@ -340,203 +325,6 @@ private:
 		}
 	}
 
-	FSlateWidgetRun::FWidgetRunInfo HandleTimeStampDisplay( const FTextRunInfo& RunInfo, const ISlateStyle* Style ) const
-	{
-		// Internal class for timestamps
-		class SChatTimeStamp : public SCompoundWidget
-		{
-			SLATE_BEGIN_ARGS(SChatTimeStamp)
-				: _DateTime(FDateTime::Now())
-			{ }
-				SLATE_ARGUMENT(FDateTime, DateTime)
-			SLATE_END_ARGS()
-
-			virtual ~SChatTimeStamp()
-			{
-				FTicker::GetCoreTicker().RemoveTicker(TickerHandle);
-			}
-
-			void Construct(const FArguments& InArgs)
-			{
-				TimeStamp = InArgs._DateTime;
-				ChildSlot
-				[
-					SNew(STextBlock)
-					.Text(this, &SChatTimeStamp::GetMessageTimeText)
-					.Font(FFriendsAndChatModuleStyle::GetStyleService()->GetNormalFont())
-				];
-
-				if (UpdateTimeStamp(60)) //this will broadcast a changed event
-				{
-					TickerHandle = FTicker::GetCoreTicker().AddTicker(FTickerDelegate::CreateSP(this, &SChatTimeStamp::UpdateTimeStamp), 60);
-				}
-			}
-
-			private:
-
-			bool UpdateTimeStamp(float Delay)
-			{
-				bool ContinueTicking = true;
-
-				FDateTime Now = FDateTime::Now();
-
-				if ((Now - TimeStamp).GetDuration() <= FTimespan::FromMinutes(1.0))
-				{
-					static FText NowTimeStamp = NSLOCTEXT("SChatWindow", "Now_TimeStamp", "Now");
-					MessageTimeAsText = NowTimeStamp;
-				}
-				else if ((Now - TimeStamp).GetDuration() <= FTimespan::FromMinutes(2.0))
-				{
-					static FText OneMinuteTimeStamp = NSLOCTEXT("SChatWindow", "1_Minute_TimeStamp", "1 min");
-					MessageTimeAsText = OneMinuteTimeStamp;
-				}
-				else if ((Now - TimeStamp).GetDuration() <= FTimespan::FromMinutes(3.0))
-				{
-					static FText TwoMinuteTimeStamp = NSLOCTEXT("SChatWindow", "2_Minute_TimeStamp", "2 min");
-					MessageTimeAsText = TwoMinuteTimeStamp;
-				}
-				else if ((Now - TimeStamp).GetDuration() <= FTimespan::FromMinutes(4.0))
-				{
-					static FText ThreeMinuteTimeStamp = NSLOCTEXT("SChatWindow", "3_Minute_TimeStamp", "3 min");
-					MessageTimeAsText = ThreeMinuteTimeStamp;
-				}
-				else if ((Now - TimeStamp).GetDuration() <= FTimespan::FromMinutes(5.0))
-				{
-					static FText FourMinuteTimeStamp = NSLOCTEXT("SChatWindow", "4_Minute_TimeStamp", "4 min");
-					MessageTimeAsText = FourMinuteTimeStamp;
-				}
-				else if ((Now - TimeStamp).GetDuration() <= FTimespan::FromMinutes(6.0))
-				{
-					static FText FiveMinuteTimeStamp = NSLOCTEXT("SChatWindow", "5_Minute_TimeStamp", "5 min");
-					MessageTimeAsText = FiveMinuteTimeStamp;
-				}
-				else
-				{
-					MessageTimeAsText = FText::AsTime(TimeStamp, EDateTimeStyle::Short, "GMT");
-					ContinueTicking = false;
-				}
-
-				return ContinueTicking;
-			}
-
-			FText GetMessageTimeText() const
-			{
-				return MessageTimeAsText;
-			}
-
-			private:
-				FText MessageTimeAsText;
-				FDelegateHandle TickerHandle;
-				FDateTime TimeStamp;
-		};
-
-		TSharedRef< SWidget > Widget = SNullWidget::NullWidget;
-
-		FDateTime TimeStamp;
-		FString TimeStampSting = RunInfo.Content.ToString();
-		if (!TimeStampSting.IsEmpty() && FDateTime::Parse(TimeStampSting, TimeStamp))
-		{
-			Widget = SNew( SChatTimeStamp )
-			.DateTime(TimeStamp);
-		}
-
-
-		TSharedRef< FSlateFontMeasure > FontMeasure = FSlateApplication::Get().GetRenderer()->GetFontMeasureService();
-		int16 Baseline = FontMeasure->GetBaseline(FFriendsAndChatModuleStyle::GetStyleService()->GetNormalFont());
-
-		return FSlateWidgetRun::FWidgetRunInfo(Widget, Baseline);
-	}
-
-	FText FormatListMessageText(TSharedRef<FChatItemViewModel> ChatItem, bool bGrouped)
-	{
-		const FText& TextToSet = ChatItem->GetMessage();
-		FFormatNamedArguments Args;
-		FString HyperlinkStyle = TEXT("UserNameTextStyle.DefaultHyperlink");
-		FSlateBrush* Brush = nullptr;
-
-		switch (ChatItem->GetMessageType())
-		{
-		case EChatMessageType::Global: 
-			HyperlinkStyle = TEXT("UserNameTextStyle.GlobalHyperlink"); 
-			Args.Add(TEXT("RoomName"), FText::FromString(TEXT("[g]")));
-			break;
-		case EChatMessageType::Whisper: 
-			HyperlinkStyle = TEXT("UserNameTextStyle.Whisperlink"); 
-			Args.Add(TEXT("RoomName"), FText::FromString(TEXT("[w]")));
-			break;
-		case EChatMessageType::Game: 
-			HyperlinkStyle = TEXT("UserNameTextStyle.GameHyperlink"); 
-			Args.Add(TEXT("RoomName"), FText::FromString(TEXT("[p]")));
-			break;
-		}
-
-		Args.Add(TEXT("Channel"), FText::FromString(EChatMessageType::ShortcutString(ChatItem->GetMessageType())));
-		Args.Add(TEXT("RecipientName"), ChatItem->GetRecipientName());
-		if(ChatItem->GetRecipientID().IsValid())
-		{
-			Args.Add(TEXT("RecipientID"), FText::FromString(ChatItem->GetRecipientID()->ToString()));
-		}
-		Args.Add(TEXT("SenderName"), ChatItem->GetSenderName());
-		if(ChatItem->GetSenderID().IsValid())
-		{
-			Args.Add(TEXT("SenderID"), FText::FromString(ChatItem->GetSenderID()->ToString()));
-		}
-
-		FString MessageTime = ChatItem->GetMessageTime().ToString();
-		FString TimeStampMarkup = "<TimeStamp>" + MessageTime + "</>";
-		Args.Add(TEXT("TimeStampMarkup"), FText::FromString(TimeStampMarkup));
-
-		Args.Add(TEXT("NameStyle"), FText::FromString(HyperlinkStyle));
-		Args.Add(TEXT("Message"), TextToSet);
-
-		FText ListMessage;
-
-		if (bGrouped)
-		{
-			ListMessage = FText::Format(NSLOCTEXT("SChatWindow", "DisplayMessage", "{Message}"), TextToSet);
-		}
-		else
-		{
-			if (!ViewModel->MultiChat())
-			{
-				if (ChatItem->IsFromSelf())
-				{
-					ListMessage = FText::Format(NSLOCTEXT("FChatItemViewModel", "DisplayNameAndMessage", "{TimeStampMarkup} {SenderName} : {Message}"), Args);
-				}
-				else
-				{
-					ListMessage = FText::Format(NSLOCTEXT("SChatWindow", "DisplayNameAndMessageWithMarkup", "{TimeStampMarkup} <a id=\"UserName\" uid=\"{SenderID}\" Username=\"{SenderName}\" style=\"{NameStyle}\">{SenderName}</> : {Message}"), Args);
-				}
-			}
-			else
-			{
-				if (ChatItem->GetMessageType() == EChatMessageType::Whisper)
-				{
-					if (ChatItem->IsFromSelf())
-					{
-						ListMessage = FText::Format(NSLOCTEXT("SChatWindow", "DisplayNameAndMessageWhisperFromSelf", "{TimeStampMarkup} <a id=\"Channel\" Command=\"{Channel}\" style=\"{NameStyle}\">{RoomName}</> to <a id=\"UserName\" uid=\"{RecipientID}\" Username=\"{RecipientName}\" style=\"{NameStyle}\">{RecipientName}</> : {Message}"), Args);
-					}
-					else
-					{
-						ListMessage = FText::Format(NSLOCTEXT("SChatWindow", "DisplayNameAndMessageWhisper", "{TimeStampMarkup} <a id=\"Channel\" Command=\"{Channel}\" style=\"{NameStyle}\">{RoomName}</> from <a id=\"UserName\" uid=\"{SenderID}\" Username=\"{SenderName}\" style=\"{NameStyle}\">{SenderName}</> : {Message}"), Args);
-					}
-				}
-				else
-				{
-					if (ChatItem->IsFromSelf())
-					{
-						ListMessage = FText::Format(NSLOCTEXT("SChatWindow", "DisplayNameAndMessageOtherFromSelf", "{TimeStampMarkup} <a id=\"Channel\" Command=\"{Channel}\" style=\"{NameStyle}\">{RoomName}</> {SenderName} : {Message}"), Args);
-					}
-					else
-					{
-						ListMessage = FText::Format(NSLOCTEXT("SChatWindow", "DisplayNameAndMessageOther", "{TimeStampMarkup} <a id=\"Channel\" Command=\"{Channel}\" style=\"{NameStyle}\">{RoomName}</> <a id=\"UserName\" uid=\"{SenderID}\" Username=\"{SenderName}\" style=\"{NameStyle}\">{SenderName}</> : {Message}"), Args);
-					}
-				}
-			}
-		}
-		return ListMessage;
-	}
-
 	EVisibility GetChatChannelVisibility() const
 	{
 		return EVisibility::Visible;
@@ -581,14 +369,12 @@ private:
 
 	TSharedPtr<SScrollBox> ChatScrollBox;
 
-	TSharedPtr<FRichTextLayoutMarshaller> RichTextMarshaller;
+	TSharedPtr<FChatTextLayoutMarshaller> RichTextMarshaller;
 
 	TSharedPtr<SMultiLineEditableTextBox> RichText;
 
 	SVerticalBox::FSlot* ChannelTextSlot;
 
-	FTextBuilder ChatText;
-
 	int32 LastMessageIndex;
 
 	TSharedPtr<FChatItemViewModel> LastMessage;
diff --git a/Engine/Source/Developer/FriendsAndChat/Private/Layers/UI/SClanItem.cpp b/Engine/Source/Developer/FriendsAndChat/Private/Layers/UI/SClanItem.cpp
--- a/Engine/Source/Developer/FriendsAndChat/Private/Layers/UI/SClanItem.cpp
+++ b/Engine/Source/Developer/FriendsAndChat/Private/Layers/UI/SClanItem.cpp
@@ -50,7 +50,8 @@ public:
 						[
 							// Primary Clan Indicator
 							SNew(SImage)
-							.Image(&FriendStyle.AwayBrush )
+							.Visibility(this, &SClanItemImpl::PrimaryClanVisibility)
+							.Image(&FriendStyle.OnlineBrush )
 						]
 					]
 					+ SHorizontalBox::Slot()
@@ -145,8 +146,31 @@ private:
 		return ActionItems;
 	}
 
+	FName ActionToItemTag(EClanActionType::Type Action)
+	{
+		return FName(*EClanActionType::ToText(Action).ToString());
+	}
+
+	EClanActionType::Type ItemTagToAction(const FName& Tag)
+	{
+		for (int32 ActionIdx = 0; ActionIdx < EClanActionType::MAX_None; ActionIdx++)
+		{
+			EClanActionType::Type ActionAsEnum = (EClanActionType::Type)ActionIdx;
+			if (Tag == ActionToItemTag(ActionAsEnum))
+			{
+				return ActionAsEnum;
+			}
+		}
+		return EClanActionType::MAX_None;
+	}
+
 	void HandleItemClicked(FName ItemTag)
 	{
+		EClanActionType::Type Action = ItemTagToAction(ItemTag);
+		if (Action != EClanActionType::MAX_None)
+		{
+			ClanInfoViewModel->PerformAction(Action);
+		}
 	}
 
 	void HandleActionMenuOpened() const
@@ -158,6 +182,11 @@ private:
 		return EVisibility::Visible;
 	}
 
+	EVisibility PrimaryClanVisibility() const
+	{
+		return ClanInfoViewModel->IsPrimaryClan() ? EVisibility::Visible : EVisibility::Hidden;
+	}
+
 	virtual bool SupportsKeyboardFocus() const override
 	{
 		return true;
diff --git a/Engine/Source/Developer/FriendsAndChat/Private/Layers/UI/SFriendActions.cpp b/Engine/Source/Developer/FriendsAndChat/Private/Layers/UI/SFriendActions.cpp
--- a/Engine/Source/Developer/FriendsAndChat/Private/Layers/UI/SFriendActions.cpp
+++ b/Engine/Source/Developer/FriendsAndChat/Private/Layers/UI/SFriendActions.cpp
@@ -22,7 +22,7 @@ public:
 
 
 		TArray<EFriendActionType::Type> Actions;
-		FriendViewModel->EnumerateActions(Actions);
+		FriendViewModel->EnumerateActions(Actions, InArgs._FromChat);
 		for (EFriendActionType::Type Action : Actions)
 		{
 			ActionContainer->AddSlot()
diff --git a/Engine/Source/Developer/FriendsAndChat/Private/Layers/UI/SFriendActions.h b/Engine/Source/Developer/FriendsAndChat/Private/Layers/UI/SFriendActions.h
--- a/Engine/Source/Developer/FriendsAndChat/Private/Layers/UI/SFriendActions.h
+++ b/Engine/Source/Developer/FriendsAndChat/Private/Layers/UI/SFriendActions.h
@@ -6,8 +6,10 @@ class SFriendActions : public SUserWidget
 {
 public:
 	SLATE_USER_ARGS(SFriendActions)
+		: _FromChat(false)
 	{ }
-	SLATE_ARGUMENT( const FFriendsAndChatStyle*, FriendStyle )
+		SLATE_ARGUMENT( const FFriendsAndChatStyle*, FriendStyle )
+		SLATE_ARGUMENT( bool, FromChat )
 	SLATE_END_ARGS()
 
 	/**
diff --git a/Engine/Source/Developer/FriendsAndChat/Private/Layers/UI/SFriendsContainer.cpp b/Engine/Source/Developer/FriendsAndChat/Private/Layers/UI/SFriendsContainer.cpp
--- a/Engine/Source/Developer/FriendsAndChat/Private/Layers/UI/SFriendsContainer.cpp
+++ b/Engine/Source/Developer/FriendsAndChat/Private/Layers/UI/SFriendsContainer.cpp
@@ -230,6 +230,11 @@ private:
 		return AddFriendButton.IsValid() && AddFriendButton->IsHovered() ? &FriendListStyle.AddFriendButtonContentHoveredBrush : &FriendListStyle.AddFriendButtonContentBrush;
 	}
 
+	EVisibility GetGlobalVisibility() const
+	{
+		return ViewModel->GetGlobalChatButtonVisibility();
+	}
+
 	EVisibility NoFriendsNoticeVisibility() const
 	{
 		for (auto& ListViewModel : ListViewModels)
@@ -267,6 +272,11 @@ private:
 
 	TSharedRef<SWidget> GetFriendsListWidget()
 	{
+		ListScrollBox = SNew(SScrollBox)
+			.ExternalScrollbar(ExternalScrollbar.ToSharedRef());
+
+		CreateLists();
+
 		// Probably move this to its own widget Nick Davies (3/25/2015)
 		return SNew(SVerticalBox)
 			+SVerticalBox::Slot()
@@ -339,7 +349,7 @@ private:
 								SNew(SBorder)
 								.Padding(0)
 								.BorderImage(FCoreStyle::Get().GetBrush("NoBorder"))
-								.Visibility(ViewModel->GetGlobalChatButtonVisibility())
+								.Visibility(this, &SFriendsContainerImpl::GetGlobalVisibility)
 								[
 									SNew(SButton)
 									.ToolTip(CreateGlobalChatToolTip())
@@ -404,77 +414,7 @@ private:
 					SNew(SOverlay)
 					+ SOverlay::Slot()
 					[
-						SNew(SScrollBox)
-						.ExternalScrollbar(ExternalScrollbar.ToSharedRef())
-						+ SScrollBox::Slot()
-						.HAlign(HAlign_Center)
-						.Padding(10)
-						.Padding(FMargin(10, 10, 40, 10))
-						[
-							SNew(STextBlock)
-							.Text(LOCTEXT("NoFriendsNotice", "Press the Plus Button to add friends."))
-							.Font(FriendFontStyle.FriendsFontSmallBold)
-							.ColorAndOpacity(FLinearColor::White)
-							.Visibility(this, &SFriendsContainerImpl::NoFriendsNoticeVisibility)
-						]
-						+SScrollBox::Slot()
-						.Padding(FMargin(10, 10, 30, 0))
-						[
-							SNew(SFriendsListContainer, ListViewModels[EFriendsDisplayLists::GameInviteDisplay].ToSharedRef())
-							.FriendStyle(&FriendListStyle)
-							.FontStyle(&FriendStyle.FriendsNormalFontStyle)
-							.Method(MenuMethod)
-							.Visibility(ListViewModels[EFriendsDisplayLists::GameInviteDisplay].Get(), &FFriendListViewModel::GetListVisibility)
-						]
-						+SScrollBox::Slot()
-						.Padding(FMargin(10, 10, 30, 0))
-						[
-							SNew(SFriendsListContainer, ListViewModels[EFriendsDisplayLists::FriendRequestsDisplay].ToSharedRef())
-							.FriendStyle(&FriendListStyle)
-							.FontStyle(&FriendStyle.FriendsNormalFontStyle)
-							.Method(MenuMethod)
-							.Visibility(ListViewModels[EFriendsDisplayLists::FriendRequestsDisplay].Get(), &FFriendListViewModel::GetListVisibility)
-						]
-						+SScrollBox::Slot()
-						.Padding(FMargin(10, 10, 30, 0))
-						[
-							SNew(SFriendsListContainer, ListViewModels[EFriendsDisplayLists::DefaultDisplay].ToSharedRef())
-							.FriendStyle(&FriendListStyle)
-							.FontStyle(&FriendStyle.FriendsNormalFontStyle)
-							.Method(MenuMethod)
-							.Visibility(ListViewModels[EFriendsDisplayLists::DefaultDisplay].Get(), &FFriendListViewModel::GetListVisibility)
-						]
-						+SScrollBox::Slot()
-						.Padding(FMargin(10, 10, 30, 0))
-						[
-							SNew(SFriendsListContainer, ListViewModels[EFriendsDisplayLists::OfflineFriends].ToSharedRef())
-							.FriendStyle(&FriendListStyle)
-							.FontStyle(&FriendStyle.FriendsNormalFontStyle)
-							.Method(MenuMethod)
-							.Visibility(ListViewModels[EFriendsDisplayLists::OfflineFriends].Get(), &FFriendListViewModel::GetListVisibility)
-						]
-						+SScrollBox::Slot()
-						.Padding(FMargin(10, 10, 30, 0))
-						[
-							SNew(SFriendsListContainer, ListViewModels[EFriendsDisplayLists::RecentPlayersDisplay].ToSharedRef())
-							.FriendStyle(&FriendListStyle)
-							.FontStyle(&FriendStyle.FriendsNormalFontStyle)
-							.Method(MenuMethod)
-							.Visibility(ListViewModels[EFriendsDisplayLists::RecentPlayersDisplay].Get(), &FFriendListViewModel::GetListVisibility)
-						]
-						+SScrollBox::Slot()
-						.Padding(FMargin(10, 10, 30, 0))
-						[
-							SNew(SFriendsListContainer, ListViewModels[EFriendsDisplayLists::OutgoingFriendInvitesDisplay].ToSharedRef())
-							.FriendStyle(&FriendListStyle)
-							.FontStyle(&FriendStyle.FriendsNormalFontStyle)
-							.Method(MenuMethod)
-							.Visibility(ListViewModels[EFriendsDisplayLists::OutgoingFriendInvitesDisplay].Get(), &FFriendListViewModel::GetListVisibility)
-						]
-						+ SScrollBox::Slot()
-						[
-							SNew(SSpacer).Size(FVector2D(10, 10))
-						]
+						ListScrollBox.ToSharedRef()
 					]
 					+ SOverlay::Slot()
 					.HAlign(HAlign_Right)
@@ -486,6 +426,93 @@ private:
 			];
 	}
 
+	virtual void RefreshRecentFriends()
+	{
+		CreateLists();
+	}
+
+	void CreateLists()
+	{
+		ListScrollBox->ClearChildren();
+
+		ListScrollBox->AddSlot()
+		.HAlign(HAlign_Center)
+		.Padding(10)
+		.Padding(FMargin(10, 10, 40, 10))
+		[
+			SNew(STextBlock)
+			.Text(LOCTEXT("NoFriendsNotice", "Press the Plus Button to add friends."))
+			.Font(FriendFontStyle.FriendsFontSmallBold)
+			.ColorAndOpacity(FLinearColor::White)
+			.Visibility(this, &SFriendsContainerImpl::NoFriendsNoticeVisibility)
+		];
+
+		ListScrollBox->AddSlot()
+		.Padding(FMargin(10, 10, 30, 0))
+		[
+			SNew(SFriendsListContainer, ListViewModels[EFriendsDisplayLists::GameInviteDisplay].ToSharedRef())
+			.FriendStyle(&FriendListStyle)
+			.FontStyle(&FriendStyle.FriendsNormalFontStyle)
+			.Method(MenuMethod)
+			.Visibility(ListViewModels[EFriendsDisplayLists::GameInviteDisplay].Get(), &FFriendListViewModel::GetListVisibility)
+		];
+
+		ListScrollBox->AddSlot()
+		.Padding(FMargin(10, 10, 30, 0))
+		[
+			SNew(SFriendsListContainer, ListViewModels[EFriendsDisplayLists::FriendRequestsDisplay].ToSharedRef())
+			.FriendStyle(&FriendListStyle)
+			.FontStyle(&FriendStyle.FriendsNormalFontStyle)
+			.Method(MenuMethod)
+			.Visibility(ListViewModels[EFriendsDisplayLists::FriendRequestsDisplay].Get(), &FFriendListViewModel::GetListVisibility)
+		];
+
+		ListScrollBox->AddSlot()
+		.Padding(FMargin(10, 10, 30, 0))
+		[
+			SNew(SFriendsListContainer, ListViewModels[EFriendsDisplayLists::DefaultDisplay].ToSharedRef())
+			.FriendStyle(&FriendListStyle)
+			.FontStyle(&FriendStyle.FriendsNormalFontStyle)
+			.Method(MenuMethod)
+			.Visibility(ListViewModels[EFriendsDisplayLists::DefaultDisplay].Get(), &FFriendListViewModel::GetListVisibility)
+		];
+
+		ListScrollBox->AddSlot()
+		.Padding(FMargin(10, 10, 30, 0))
+		[
+			SNew(SFriendsListContainer, ListViewModels[EFriendsDisplayLists::OfflineFriends].ToSharedRef())
+			.FriendStyle(&FriendListStyle)
+			.FontStyle(&FriendStyle.FriendsNormalFontStyle)
+			.Method(MenuMethod)
+			.Visibility(ListViewModels[EFriendsDisplayLists::OfflineFriends].Get(), &FFriendListViewModel::GetListVisibility)
+		];
+	
+		ListScrollBox->AddSlot()
+		.Padding(FMargin(10, 10, 30, 0))
+		[
+			SNew(SFriendsListContainer, ListViewModels[EFriendsDisplayLists::RecentPlayersDisplay].ToSharedRef())
+			.FriendStyle(&FriendListStyle)
+			.FontStyle(&FriendStyle.FriendsNormalFontStyle)
+			.Method(MenuMethod)
+			.Visibility(ListViewModels[EFriendsDisplayLists::RecentPlayersDisplay].Get(), &FFriendListViewModel::GetListVisibility)
+		];
+
+		ListScrollBox->AddSlot()
+		.Padding(FMargin(10, 10, 30, 0))
+		[
+			SNew(SFriendsListContainer, ListViewModels[EFriendsDisplayLists::OutgoingFriendInvitesDisplay].ToSharedRef())
+			.FriendStyle(&FriendListStyle)
+			.FontStyle(&FriendStyle.FriendsNormalFontStyle)
+			.Method(MenuMethod)
+			.Visibility(ListViewModels[EFriendsDisplayLists::OutgoingFriendInvitesDisplay].Get(), &FFriendListViewModel::GetListVisibility)
+		];
+
+		ListScrollBox->AddSlot()
+			[
+				SNew(SSpacer).Size(FVector2D(10, 10))
+			];
+	}
+
 
 FSlateColor GetFriendTabTextColor(int32 WidgetIndex) const
 {
@@ -541,6 +568,8 @@ private:
 	TSharedPtr<SScrollBar> ExternalScrollbar;
 
 	TSharedPtr<SButton> AddFriendButton;
+
+	TSharedPtr<SScrollBox> ListScrollBox;
 };
 
 TSharedRef<SFriendsContainer> SFriendsContainer::New()
diff --git a/Engine/Source/Developer/FriendsAndChat/Private/Models/FriendFriendItem.cpp b/Engine/Source/Developer/FriendsAndChat/Private/Models/FriendFriendItem.cpp
new file mode 100644
--- /dev/null
+++ b/Engine/Source/Developer/FriendsAndChat/Private/Models/FriendFriendItem.cpp
@@ -0,0 +1,50 @@
+// Copyright 1998-2015 Epic Games, Inc. All Rights Reserved.
+
+#include "FriendsAndChatPrivatePCH.h"
+#include "FriendFriendItem.h"
+
+#include "FriendsService.h"
+
+#define LOCTEXT_NAMESPACE ""
+
+
+// FFriendStruct implementation
+
+bool FFriendFriendItem::IsGameJoinable() const
+{
+	if (GetOnlineFriend().IsValid())
+	{
+		const FOnlineUserPresence& FriendPresence = GetOnlineFriend()->GetPresence();
+		const bool bIsOnline = FriendPresence.Status.State != EOnlinePresenceState::Offline;
+		if (bIsOnline)
+		{
+			const bool bIsGameJoinable = FriendPresence.bIsJoinable && !FriendsService->IsFriendInSameSession(AsShared());
+			if (bIsGameJoinable)
+			{
+				return true;
+			}
+		}
+	}
+	return false;
+}
+
+bool FFriendFriendItem::CanJoinParty() const
+{
+	TSharedPtr<IOnlinePartyJoinInfo> PartyInfo = GetPartyJoinInfo();
+	return PartyInfo.IsValid() && PartyInfo->GetIsAcceptingMembers() && !FriendsService->IsFriendInSameParty(AsShared());
+}
+
+bool FFriendFriendItem::CanInvite() const
+{
+	FString FriendsClientID = GetClientId();
+	return FriendsClientID == FriendsService->GetUserClientId() || FFriendItem::LauncherClientIds.Contains(FriendsClientID);
+}
+
+TSharedPtr<IOnlinePartyJoinInfo> FFriendFriendItem::GetPartyJoinInfo() const
+{
+	// obtain party info from presence
+	return FriendsService->GetPartyJoinInfo(AsShared());
+}
+
+#undef LOCTEXT_NAMESPACE
+
diff --git a/Engine/Source/Developer/FriendsAndChat/Private/Models/FriendFriendItem.h b/Engine/Source/Developer/FriendsAndChat/Private/Models/FriendFriendItem.h
new file mode 100644
--- /dev/null
+++ b/Engine/Source/Developer/FriendsAndChat/Private/Models/FriendFriendItem.h
@@ -0,0 +1,30 @@
+// Copyright 1998-2015 Epic Games, Inc. All Rights Reserved.
+
+#pragma once
+
+#include "IFriendItem.h"
+/**
+ * Class containing the friend information - used to build the list view.
+ */
+class FFriendFriendItem : public FFriendItem
+{
+public:
+
+	// FFriendStuct overrides
+
+	virtual bool IsGameJoinable() const override;
+	virtual bool CanJoinParty() const override;
+	virtual bool CanInvite() const override;
+	virtual TSharedPtr<IOnlinePartyJoinInfo> GetPartyJoinInfo() const override;
+
+	// FFriendFriendItem
+
+	FFriendFriendItem(TSharedPtr< FOnlineFriend > InOnlineFriend, TSharedPtr< FOnlineUser > InOnlineUser, EFriendsDisplayLists::Type InListType, const TSharedRef<class FFriendsService>& InFriendsService)
+		: FFriendItem(InOnlineFriend, InOnlineUser, InListType)
+		, FriendsService(InFriendsService)
+	{ }
+
+private:
+
+	TSharedPtr<class FFriendsService> FriendsService;
+};
diff --git a/Engine/Source/Developer/FriendsAndChat/Private/Models/FriendGameInviteItem.cpp b/Engine/Source/Developer/FriendsAndChat/Private/Models/FriendGameInviteItem.cpp
--- a/Engine/Source/Developer/FriendsAndChat/Private/Models/FriendGameInviteItem.cpp
+++ b/Engine/Source/Developer/FriendsAndChat/Private/Models/FriendGameInviteItem.cpp
@@ -3,8 +3,28 @@
 #include "FriendsAndChatPrivatePCH.h"
 #include "FriendGameInviteItem.h"
 
+#include "GameAndPartyService.h"
+
 // FFriendGameInviteItem
 
+bool FFriendGameInviteItem::CanJoinParty() const
+{
+	TSharedPtr<IOnlinePartyJoinInfo> PartyInfo = GetPartyJoinInfo();
+	return PartyInfo.IsValid() && PartyInfo->GetIsAcceptingMembers() && !GameAndPartyService->IsFriendInSameParty(AsShared());
+}
+
+bool FFriendGameInviteItem::CanInvite() const
+{
+	FString FriendsClientID = GetClientId();
+	return FriendsClientID == GameAndPartyService->GetUserClientId() || FFriendItem::LauncherClientIds.Contains(FriendsClientID);
+}
+
+TSharedPtr<IOnlinePartyJoinInfo> FFriendGameInviteItem::GetPartyJoinInfo() const
+{
+	// obtain party info from presence
+	return GameAndPartyService->GetPartyJoinInfo(AsShared());
+}
+
 bool FFriendGameInviteItem::IsGameRequest() const
 {
 	return true;
@@ -13,7 +33,7 @@ bool FFriendGameInviteItem::IsGameRequest() const
 bool FFriendGameInviteItem::IsGameJoinable() const
 {
 	const TSharedPtr<const IFriendItem> Item = AsShared();
-	return !FriendsAndChatManager.Pin()->IsFriendInSameSession(Item);
+	return !GameAndPartyService->IsFriendInSameSession(Item);
 }
 
 TSharedPtr<const FUniqueNetId> FFriendGameInviteItem::GetGameSessionId() const
@@ -21,7 +41,7 @@ TSharedPtr<const FUniqueNetId> FFriendGameInviteItem::GetGameSessionId() const
 	TSharedPtr<const FUniqueNetId> SessionId = NULL;
 	if (SessionResult->Session.SessionInfo.IsValid())
 	{
-		SessionId = FriendsAndChatManager.Pin()->GetGameSessionId(SessionResult->Session.SessionInfo->GetSessionId().ToString());
+		SessionId = GameAndPartyService->GetGameSessionId(SessionResult->Session.SessionInfo->GetSessionId().ToString());
 	}
 	return SessionId;
 }
diff --git a/Engine/Source/Developer/FriendsAndChat/Private/Models/FriendGameInviteItem.h b/Engine/Source/Developer/FriendsAndChat/Private/Models/FriendGameInviteItem.h
--- a/Engine/Source/Developer/FriendsAndChat/Private/Models/FriendGameInviteItem.h
+++ b/Engine/Source/Developer/FriendsAndChat/Private/Models/FriendGameInviteItem.h
@@ -13,6 +13,9 @@ public:
 
 	// FFriendStuct overrides
 
+	virtual bool CanJoinParty() const override;
+	virtual bool CanInvite() const override;
+	virtual TSharedPtr<IOnlinePartyJoinInfo> GetPartyJoinInfo() const override;
 	virtual bool IsGameRequest() const override;
 	virtual bool IsGameJoinable() const override;
 	virtual TSharedPtr<const FUniqueNetId> GetGameSessionId() const override;
@@ -24,16 +27,16 @@ public:
 		const TSharedRef<FOnlineSessionSearchResult>& InSessionResult,
 		const FString& InClientId,
 		const TSharedRef<FOnlineFriend>& InOnlineFriend,
-		const TSharedRef<class FFriendsAndChatManager> InFriendsAndChatManager)
-		: FFriendItem(InOnlineFriend, InOnlineUser, EFriendsDisplayLists::GameInviteDisplay, InFriendsAndChatManager)
+		const TSharedRef<class FGameAndPartyService>& InGameAndPartyService)
+		: FFriendItem(InOnlineFriend, InOnlineUser, EFriendsDisplayLists::GameInviteDisplay)
 		, SessionResult(InSessionResult)
 		, ClientId(InClientId)
-		, FriendsAndChatManager(InFriendsAndChatManager)
+		, GameAndPartyService(InGameAndPartyService)
 	{ }
 
 private:
 
 	TSharedPtr<FOnlineSessionSearchResult> SessionResult;
 	FString ClientId;
-	TWeakPtr<class FFriendsAndChatManager> FriendsAndChatManager;
+	TSharedPtr<class FGameAndPartyService> GameAndPartyService;
 };
diff --git a/Engine/Source/Developer/FriendsAndChat/Private/Models/FriendItem.cpp b/Engine/Source/Developer/FriendsAndChat/Private/Models/FriendItem.cpp
--- a/Engine/Source/Developer/FriendsAndChat/Private/Models/FriendItem.cpp
+++ b/Engine/Source/Developer/FriendsAndChat/Private/Models/FriendItem.cpp
@@ -111,6 +111,11 @@ const FString FFriendItem::GetClientName() const
 	return Result;
 }
 
+FDateTime FFriendItem::GetLastSeen() const
+{
+	return FDateTime::UtcNow();
+}
+
 const bool FFriendItem::IsOnline() const
 {
 	if(OnlineFriend.IsValid())
@@ -131,19 +136,6 @@ const EOnlinePresenceState::Type FFriendItem::GetOnlineStatus() const
 
 bool FFriendItem::IsGameJoinable() const
 {
-	if (OnlineFriend.IsValid())
-	{
-		const FOnlineUserPresence& FriendPresence = OnlineFriend->GetPresence();
-		const bool bIsOnline = FriendPresence.Status.State != EOnlinePresenceState::Offline;
-		if (bIsOnline)
-		{
-			const bool bIsGameJoinable = FriendPresence.bIsJoinable && !FriendsAndChatManager.Pin()->IsFriendInSameSession(AsShared());
-			if (bIsGameJoinable)
-			{
-				return true;
-			}
-		}
-	}
 	return false;
 }
 
@@ -155,14 +147,12 @@ bool FFriendItem::IsInParty() const
 
 bool FFriendItem::CanJoinParty() const
 {
-	TSharedPtr<IOnlinePartyJoinInfo> PartyInfo = GetPartyJoinInfo();
-	return PartyInfo.IsValid() && PartyInfo->GetIsAcceptingMembers() && !FriendsAndChatManager.Pin()->IsFriendInSameParty(AsShared());
+	return false;
 }
 
 bool FFriendItem::CanInvite() const
 {
-	FString FriendsClientID = GetClientId();
-	return FriendsClientID == FriendsAndChatManager.Pin()->GetUserClientId() || FFriendItem::LauncherClientIds.Contains(FriendsClientID);
+	return false;
 }
 
 TSharedPtr<const FUniqueNetId> FFriendItem::GetGameSessionId() const
@@ -180,8 +170,7 @@ TSharedPtr<const FUniqueNetId> FFriendItem::GetGameSessionId() const
 
 TSharedPtr<IOnlinePartyJoinInfo> FFriendItem::GetPartyJoinInfo() const
 {
-	// obtain party info from presence
-	return FriendsAndChatManager.Pin()->GetPartyJoinInfo(AsShared());
+	return nullptr;
 }
 
 const TSharedRef<const FUniqueNetId> FFriendItem::GetUniqueID() const
diff --git a/Engine/Source/Developer/FriendsAndChat/Private/Models/FriendItem.h b/Engine/Source/Developer/FriendsAndChat/Private/Models/FriendItem.h
--- a/Engine/Source/Developer/FriendsAndChat/Private/Models/FriendItem.h
+++ b/Engine/Source/Developer/FriendsAndChat/Private/Models/FriendItem.h
@@ -21,9 +21,8 @@ public:
 	 * @param InOnlineUser The online user.
 	 * @param InListType The list type.
 	 */
-	FFriendItem(TSharedPtr< FOnlineFriend > InOnlineFriend, TSharedPtr< FOnlineUser > InOnlineUser, EFriendsDisplayLists::Type InListType, const TSharedRef<class FFriendsAndChatManager>& InFriendsAndChatManager)
-		: FriendsAndChatManager(InFriendsAndChatManager)
-		, bIsUpdated(true)
+	FFriendItem(TSharedPtr< FOnlineFriend > InOnlineFriend, TSharedPtr< FOnlineUser > InOnlineUser, EFriendsDisplayLists::Type InListType)
+		: bIsUpdated(true)
 		, GroupName(TEXT(""))
 		, OnlineFriend( InOnlineFriend )
 		, OnlineUser( InOnlineUser )
@@ -85,6 +84,12 @@ public:
 	*/
 	virtual const FString GetClientName() const override;
 
+	/**
+	* Get last seen time
+	* @return The time
+	*/
+	virtual FDateTime GetLastSeen() const override;
+
 	/**
 	 * Get if the user is online.
 	 * @return The user online state.
@@ -205,8 +210,6 @@ protected:
 		, GroupName(TEXT(""))
 	{ };
 
-	TWeakPtr<class FFriendsAndChatManager> FriendsAndChatManager;
-
 private:
 
 	/** Holds if this item has been updated. */
diff --git a/Engine/Source/Developer/FriendsAndChat/Private/Models/FriendPartyInviteItem.cpp b/Engine/Source/Developer/FriendsAndChat/Private/Models/FriendPartyInviteItem.cpp
--- a/Engine/Source/Developer/FriendsAndChat/Private/Models/FriendPartyInviteItem.cpp
+++ b/Engine/Source/Developer/FriendsAndChat/Private/Models/FriendPartyInviteItem.cpp
@@ -3,8 +3,22 @@
 #include "FriendsAndChatPrivatePCH.h"
 #include "FriendPartyInviteItem.h"
 
+#include "GameAndPartyService.h"
+
 // FFriendPartyInviteItem
 
+bool FFriendPartyInviteItem::CanJoinParty() const
+{
+	TSharedPtr<IOnlinePartyJoinInfo> PartyInfo = GetPartyJoinInfo();
+	return PartyInfo.IsValid() && PartyInfo->GetIsAcceptingMembers() && !GameAndPartyService->IsFriendInSameParty(AsShared());
+}
+
+bool FFriendPartyInviteItem::CanInvite() const
+{
+	FString FriendsClientID = GetClientId();
+	return FriendsClientID == GameAndPartyService->GetUserClientId() || FFriendItem::LauncherClientIds.Contains(FriendsClientID);
+}
+
 bool FFriendPartyInviteItem::IsGameRequest() const
 {
 	return true;
@@ -17,7 +31,7 @@ bool FFriendPartyInviteItem::IsGameJoinable() const
 	if (PartyJoinInfo.IsValid())
 	{
 		const TSharedPtr<const IFriendItem> Item = AsShared();
-		return !FriendsAndChatManager.Pin()->IsFriendInSameParty(Item);
+		return !GameAndPartyService->IsFriendInSameParty(Item);
 	}
 	return false;
 }
diff --git a/Engine/Source/Developer/FriendsAndChat/Private/Models/FriendPartyInviteItem.h b/Engine/Source/Developer/FriendsAndChat/Private/Models/FriendPartyInviteItem.h
--- a/Engine/Source/Developer/FriendsAndChat/Private/Models/FriendPartyInviteItem.h
+++ b/Engine/Source/Developer/FriendsAndChat/Private/Models/FriendPartyInviteItem.h
@@ -13,6 +13,8 @@ public:
 
 	// FFriendStuct overrides
 
+	virtual bool CanJoinParty() const override;
+	virtual bool CanInvite() const override;
 	virtual bool IsGameRequest() const override;
 	virtual bool IsGameJoinable() const override;
 	virtual TSharedPtr<const FUniqueNetId> GetGameSessionId() const override;
@@ -25,10 +27,11 @@ public:
 		const TSharedRef<FOnlineUser>& InOnlineUser,
 		const FString& InClientId,
 		const TSharedRef<IOnlinePartyJoinInfo>& InPartyJoinInfo,
-		const TSharedRef<class FFriendsAndChatManager>& FriendsAndChatManager)
-		: FFriendItem(nullptr, InOnlineUser, EFriendsDisplayLists::GameInviteDisplay, FriendsAndChatManager)
+		const TSharedRef<class FGameAndPartyService>& InGameAndPartyService)
+		: FFriendItem(nullptr, InOnlineUser, EFriendsDisplayLists::GameInviteDisplay)
 		, PartyJoinInfo(InPartyJoinInfo)
 		, ClientId(InClientId)
+		, GameAndPartyService(InGameAndPartyService)
 	{ }
 
 protected:
@@ -41,4 +44,5 @@ private:
 
 	TSharedPtr<IOnlinePartyJoinInfo> PartyJoinInfo;
 	FString ClientId;
+	TSharedPtr<class FGameAndPartyService> GameAndPartyService;
 };
diff --git a/Engine/Source/Developer/FriendsAndChat/Private/Models/FriendRecentPlayerItem.cpp b/Engine/Source/Developer/FriendsAndChat/Private/Models/FriendRecentPlayerItem.cpp
--- a/Engine/Source/Developer/FriendsAndChat/Private/Models/FriendRecentPlayerItem.cpp
+++ b/Engine/Source/Developer/FriendsAndChat/Private/Models/FriendRecentPlayerItem.cpp
@@ -44,6 +44,15 @@ const FString FFriendRecentPlayerItem::GetClientName() const
 	return TEXT("");
 }
 
+FDateTime FFriendRecentPlayerItem::GetLastSeen() const
+{
+	if (RecentPlayer.IsValid())
+	{
+		return RecentPlayer->GetLastSeen();
+	}
+	return FDateTime::UtcNow();
+}
+
 const bool FFriendRecentPlayerItem::IsOnline() const
 {
 	return false;
diff --git a/Engine/Source/Developer/FriendsAndChat/Private/Models/FriendRecentPlayerItems.h b/Engine/Source/Developer/FriendsAndChat/Private/Models/FriendRecentPlayerItems.h
--- a/Engine/Source/Developer/FriendsAndChat/Private/Models/FriendRecentPlayerItems.h
+++ b/Engine/Source/Developer/FriendsAndChat/Private/Models/FriendRecentPlayerItems.h
@@ -66,6 +66,13 @@ public:
 	*/
 	virtual const FString GetClientName() const override;
 
+
+	/**
+	* Get last seen time
+	* @return The date last seen
+	*/
+	virtual FDateTime GetLastSeen() const override;
+
 	/**
 	 * Get if the user is online.
 	 * @return The user online state.
diff --git a/Engine/Source/Developer/FriendsAndChat/Private/Models/IFriendList.h b/Engine/Source/Developer/FriendsAndChat/Private/Models/IFriendList.h
--- a/Engine/Source/Developer/FriendsAndChat/Private/Models/IFriendList.h
+++ b/Engine/Source/Developer/FriendsAndChat/Private/Models/IFriendList.h
@@ -32,6 +32,16 @@ protected:
 			return ( A->GetName() < B->GetName() );
 		}
 	};
+
+	struct FCompareGroupByLastSeen
+	{
+		FORCEINLINE bool operator()(const TSharedPtr< IFriendItem > A, const TSharedPtr< IFriendItem > B) const
+		{
+			check(A.IsValid());
+			check(B.IsValid());
+			return (A->GetLastSeen() > B->GetLastSeen());
+		}
+	};
 };
 
 class IFriendListFactory
@@ -44,4 +54,5 @@ public:
 
 FACTORY(TSharedRef< IFriendListFactory >, FFriendListFactory,
 	const TSharedRef<class IFriendViewModelFactory>& FriendViewModelFactory,
-	const TSharedRef<class FFriendsAndChatManager>& FriendsAndChatManager);
+	const TSharedRef<class FFriendsService>& FriendsService,
+	const TSharedRef<class FGameAndPartyService>& GamePartyInviteService);
diff --git a/Engine/Source/Developer/FriendsAndChat/Public/FriendsAndChat.h b/Engine/Source/Developer/FriendsAndChat/Public/FriendsAndChat.h
--- a/Engine/Source/Developer/FriendsAndChat/Public/FriendsAndChat.h
+++ b/Engine/Source/Developer/FriendsAndChat/Public/FriendsAndChat.h
@@ -30,3 +30,6 @@
 #include "FriendsChatStyle.h"
 #include "FriendsChatChromeStyle.h"
 #include "FriendsAndChatStyle.h"
+#include "IChatNotificationService.h"
+#include "IChatCommunicationService.h"
+#include "IGameAndPartyService.h"
diff --git a/Engine/Source/Developer/FriendsAndChat/Public/FriendsChatStyle.h b/Engine/Source/Developer/FriendsAndChat/Public/FriendsChatStyle.h
--- a/Engine/Source/Developer/FriendsAndChat/Public/FriendsChatStyle.h
+++ b/Engine/Source/Developer/FriendsAndChat/Public/FriendsChatStyle.h
@@ -58,6 +58,11 @@ struct FRIENDSANDCHAT_API FFriendsChatStyle
 	FTextBlockStyle TextStyle;
 	FFriendsChatStyle& SetTextStyle(const FTextBlockStyle& InTextStyle);
 
+	/** Time Stamp Text Style */
+	UPROPERTY(EditAnywhere, Category=Appearance)
+	FTextBlockStyle TimeStampTextStyle;
+	FFriendsChatStyle& SetTimeStampTextStyle(const FTextBlockStyle& InTextStyle);
+
 	UPROPERTY(EditAnywhere, Category=Appearance)
 	FLinearColor DefaultChatColor;
 	FFriendsChatStyle& SetDefaultChatColor(const FLinearColor& InFontColor);
diff --git a/Engine/Source/Developer/FriendsAndChat/Public/Interfaces/IFriendsAndChatManager.h b/Engine/Source/Developer/FriendsAndChat/Public/Interfaces/IFriendsAndChatManager.h
--- a/Engine/Source/Developer/FriendsAndChat/Public/Interfaces/IFriendsAndChatManager.h
+++ b/Engine/Source/Developer/FriendsAndChat/Public/Interfaces/IFriendsAndChatManager.h
@@ -20,6 +20,20 @@ class IFriendsAndChatManager
 {
 public:
 
+	/** Log out and kill the friends list window. */
+	virtual void Logout() = 0;
+
+	/** Log in and start checking for Friends. */
+	virtual void Login(IOnlineSubsystem* InOnlineSub, bool bInIsGame = false, int32 LocalUserID = 0) = 0;
+
+	/** Is the chat manager logged in. */
+	virtual bool IsLoggedIn() = 0;
+
+	/** Set the user to an online state. */
+	virtual void SetOnline() = 0;
+
+	/** Set the user to an away state. */
+	virtual void SetAway() = 0;
 
 	/**
 	* Create the chat Settings service used to customize chat
@@ -88,6 +102,24 @@ public:
 	*/
 	virtual TSharedPtr<class FFriendsNavigationService> GetNavigationService() = 0;
 
+	/**
+	* Get the notification service.
+	* @return The notification service
+	*/
+	virtual TSharedPtr<class IChatNotificationService> GetNotificationService() = 0;
+
+	/**
+	* Get the Message service.
+	* @return The message service
+	*/
+	virtual TSharedPtr<class IChatCommunicationService> GetCommunicationService() = 0;
+
+	/**
+	* Get the GameAndParty service.
+	* @return The GameAndParty service
+	*/
+	virtual TSharedPtr<class IGameAndPartyService> GetGameAndPartyService() = 0;
+
 	/**
 	 * Insert a network chat message.
 	 * @param InMessage The chat message.
@@ -105,21 +137,6 @@ public:
 	 */
 	virtual void OnChatPublicRoomJoined(const FString& ChatRoomID) = 0;
 
-	/** Log out and kill the friends list window. */
-	virtual void Logout() = 0;
-
-	/** Log in and start checking for Friends. */
-	virtual void Login(IOnlineSubsystem* InOnlineSub, bool bInIsGame = false, int32 LocalUserID = 0) = 0;
-
-	/** Is the chat manager logged in. */
-	virtual bool IsLoggedIn() = 0;
-
-	/** Set the user to an online state. */
-	virtual void SetOnline() = 0;
-
-	/** Set the user to an away state. */
-	virtual void SetAway() = 0;
-
 	/**
 	* Get the online status
 	*
@@ -127,29 +144,6 @@ public:
 	*/
 	virtual EOnlinePresenceState::Type GetOnlineStatus() = 0;
 
-	/**
-	* Get the friends filtered list of friends.
-	*
-	* @param OutFriendsList  Array of friends to fill in.
-	* @return the friend list count.
-	*/
-	virtual int32 GetFilteredFriendsList(TArray< TSharedPtr< class IFriendItem > >& OutFriendsList) = 0;
-
-	/**
-	* Get the recent players list.
-	* @return the list.
-	*/
-
-	virtual TArray< TSharedPtr< class IFriendItem > >& GetRecentPlayerList() = 0;
-
-	/**
-	* Get incoming game invite list.
-	*
-	* @param OutFriendsList  Array of friends to fill in.
-	* @return The friend list count.
-	*/
-	virtual int32 GetFilteredGameInviteList(TArray< TSharedPtr< class IFriendItem > >& OutFriendsList) = 0;
-
 	/** 
 	 * Set the application view model to query and perform actions on.
 	 * @param ClientID The ID of the application
@@ -164,33 +158,9 @@ public:
 
 	virtual void ClearApplicationViewModels() = 0;
 
-	DECLARE_EVENT_OneParam(IFriendsAndChatManager, FOnFriendsNotificationEvent, const bool /*Show or Clear */)
-	virtual FOnFriendsNotificationEvent& OnFriendsNotification() = 0;
-
-	DECLARE_EVENT_OneParam(IFriendsAndChatManager, FOnFriendsNotificationActionEvent, TSharedRef<FFriendsAndChatMessage> /*Chat notification*/)
-	virtual FOnFriendsNotificationActionEvent& OnFriendsActionNotification() = 0;
-
-	DECLARE_EVENT_OneParam(IFriendsAndChatManager, FOnFriendsUserSettingsUpdatedEvent, /*struct*/ FFriendsAndChatSettings& /* New Options */)
-	virtual FOnFriendsUserSettingsUpdatedEvent& OnFriendsUserSettingsUpdated() = 0;
-
-	DECLARE_EVENT_TwoParams(IFriendsAndChatManager, FOnFriendsJoinGameEvent, const FUniqueNetId& /*FriendId*/, const FUniqueNetId& /*SessionId*/)
-	virtual FOnFriendsJoinGameEvent& OnFriendsJoinGame() = 0;
-
-	DECLARE_EVENT_ThreeParams(IFriendsAndChatManager, FOnFriendsJoinPartyEvent, const FUniqueNetId& /*SenderId*/, const TSharedRef<class IOnlinePartyJoinInfo>& /*PartyJoinInfo*/, bool /*bIsFromInvite*/)
-	virtual FOnFriendsJoinPartyEvent& OnFriendsJoinParty() = 0;
-
-	DECLARE_EVENT_TwoParams(IFriendsAndChatManager, FChatMessageReceivedEvent, EChatMessageType::Type /*Type of message received*/, TSharedPtr<IFriendItem> /*Friend if chat type is whisper*/);
-	virtual FChatMessageReceivedEvent& OnChatMessageRecieved() = 0;
-
 	DECLARE_DELEGATE_RetVal(bool, FAllowFriendsJoinGame);
 	virtual FAllowFriendsJoinGame& AllowFriendsJoinGame() = 0;
 
-	DECLARE_EVENT(IFriendsAndChatManager, FOnFriendsUpdated)
-	virtual FOnFriendsUpdated& OnFriendsListUpdated() = 0;
-
-	DECLARE_EVENT(IFriendsAndChatManager, FOnGameInvitesUpdated)
-	virtual FOnGameInvitesUpdated& OnGameInvitesUpdated() = 0;
-
 public:
 
 	/** Virtual destructor. */
diff --git a/Engine/Source/Developer/FriendsAndChat/Public/Layers/Presentation/IChatCommunicationService.h b/Engine/Source/Developer/FriendsAndChat/Public/Layers/Presentation/IChatCommunicationService.h
new file mode 100644
--- /dev/null
+++ b/Engine/Source/Developer/FriendsAndChat/Public/Layers/Presentation/IChatCommunicationService.h
@@ -0,0 +1,21 @@
+// Copyright 1998-2015 Epic Games, Inc. All Rights Reserved.
+
+#pragma once
+
+/**
+ * Interface for the message communication service
+ */
+class IChatCommunicationService
+{
+public:
+
+	virtual bool SendRoomMessage(const FString& RoomName, const FString& MsgBody) = 0;
+	virtual bool SendPrivateMessage(TSharedPtr<IFriendItem> Friend, const FText MessageText) = 0;
+	virtual void InsertNetworkMessage(const FString& MsgBody) = 0;
+
+	DECLARE_EVENT_OneParam(IChatCommunicationService, FOnChatMessageAddedEvent, const TSharedRef<struct FFriendChatMessage> /*The chat message*/)
+	virtual FOnChatMessageAddedEvent& OnChatMessageAdded() = 0;
+
+	DECLARE_EVENT_TwoParams(IChatCommunicationService, FOnChatMessageReceivedEvent, EChatMessageType::Type /*Type of message received*/, TSharedPtr<IFriendItem> /*Friend if chat type is whisper*/);
+	virtual FOnChatMessageReceivedEvent& OnChatMessageRecieved() = 0;
+};
diff --git a/Engine/Source/Developer/FriendsAndChat/Public/Layers/Presentation/IChatNotificationService.h b/Engine/Source/Developer/FriendsAndChat/Public/Layers/Presentation/IChatNotificationService.h
new file mode 100644
--- /dev/null
+++ b/Engine/Source/Developer/FriendsAndChat/Public/Layers/Presentation/IChatNotificationService.h
@@ -0,0 +1,19 @@
+// Copyright 1998-2015 Epic Games, Inc. All Rights Reserved.
+
+#pragma once
+
+class IChatNotificationService
+{
+public:
+
+	virtual void SendNotification(const FString& InMessage, EMessageType::Type) = 0;
+
+	DECLARE_EVENT_OneParam(IChatNotificationService, FOnNotificationAvailableEvent, bool)
+	virtual FOnNotificationAvailableEvent& OnNotificationsAvailable() = 0;
+
+	DECLARE_EVENT_OneParam(IChatNotificationService, FOnSendNotificationEvent, TSharedRef<FFriendsAndChatMessage> /*Chat notification*/)
+	virtual FOnSendNotificationEvent& OnSendNotification() = 0;
+
+	DECLARE_MULTICAST_DELEGATE_TwoParams(FOnNotificationResponse, TSharedPtr<FFriendsAndChatMessage>, EResponseType::Type);
+	typedef FOnNotificationResponse::FDelegate FOnNotificationResponseDelegate;
+};
\ No newline at end of file
diff --git a/Engine/Source/Developer/FriendsAndChat/Public/Layers/Presentation/IGameAndPartyService.h b/Engine/Source/Developer/FriendsAndChat/Public/Layers/Presentation/IGameAndPartyService.h
new file mode 100644
--- /dev/null
+++ b/Engine/Source/Developer/FriendsAndChat/Public/Layers/Presentation/IGameAndPartyService.h
@@ -0,0 +1,16 @@
+// Copyright 1998-2015 Epic Games, Inc. All Rights Reserved.
+/*=============================================================================
+	IGameAndPartyService.h: Interface for game and party services.
+	The game and party service is used to inform when a party action has been performed.
+	e.g. friend joined your party
+=============================================================================*/
+
+#pragma once
+
+class IGameAndPartyService
+{
+public:
+
+	DECLARE_EVENT_ThreeParams(IGameAndPartyService, FOnFriendsJoinPartyEvent, const FUniqueNetId& /*SenderId*/, const TSharedRef<class IOnlinePartyJoinInfo>& /*PartyJoinInfo*/, bool /*bIsFromInvite*/)
+	virtual FOnFriendsJoinPartyEvent& OnFriendsJoinParty() = 0;
+};
diff --git a/Engine/Source/Developer/FriendsAndChat/Public/Models/FriendsAndChatMessage.h b/Engine/Source/Developer/FriendsAndChat/Public/Models/FriendsAndChatMessage.h
--- a/Engine/Source/Developer/FriendsAndChat/Public/Models/FriendsAndChatMessage.h
+++ b/Engine/Source/Developer/FriendsAndChat/Public/Models/FriendsAndChatMessage.h
@@ -2,12 +2,8 @@
 
 #pragma once
 
-
-class FUniqueNetId;
-
-
-// Enum to list download status flags
-namespace EFriendsRequestType
+// Enum to list message flags
+namespace EMessageType
 {
 	enum Type
 	{
@@ -25,12 +21,13 @@ namespace EFriendsRequestType
 		FriendAccepted,
 		// The message is handled. Clear the notification if present
 		MessageHandled,
+		// Invite someone to join your clan
+		ClanInvite,
 	};
 };
 
-
-// Enum to list friend response type
-namespace EFriendsResponseType
+// Enum to list message response types
+namespace EResponseType
 {
 	enum Type
 	{
@@ -43,6 +40,7 @@ namespace EFriendsResponseType
 	};
 };
 
+class FUniqueNetId;
 
 // Enum holding the display list selection
 namespace EFriendsDisplayLists
@@ -131,6 +129,26 @@ public:
 		return UniqueFriendID.ToSharedRef();
 	}
 
+	/**
+	* Get the Net ID for the clan this message is sent to.
+	*
+	* @return The net ID.
+	*/
+	const TSharedRef<const FUniqueNetId> GetClanID()
+	{
+		return UniqueClanID.ToSharedRef();
+	}
+
+	/**
+	* Set the Net ID for the clan this message is sent to.
+	*
+	* @return The net ID.
+	*/
+	void SetClanID(TSharedRef<const FUniqueNetId> ClanID)
+	{
+		UniqueClanID = ClanID;
+	}
+
 	/**
 	 * Set a button callbacks.
 	 *
@@ -166,7 +184,7 @@ public:
 	 *
 	 * @param InType The message type.
 	 */
-	void SetMessageType( EFriendsRequestType::Type InType )
+	void SetMessageType(EMessageType::Type InType)
 	{
 		MessageType = InType;
 	}
@@ -196,7 +214,7 @@ public:
 	 *
 	 * @return The message type.
 	 */
-	EFriendsRequestType::Type GetMessageType()
+	EMessageType::Type GetMessageType()
 	{
 		return MessageType;
 	}
@@ -290,7 +308,7 @@ public:
 	/** Set this message into a handled state. */
 	void SetHandled()
 	{
-		MessageType = EFriendsRequestType::MessageHandled;
+		MessageType = EMessageType::MessageHandled;
 	}
 
 private:
@@ -309,9 +327,12 @@ private:
 	
 	// Holds the Unique Friend Net ID
 	const TSharedPtr<const FUniqueNetId> UniqueFriendID;
+
+	// Holds the Unique clan Net ID
+	TSharedPtr<const FUniqueNetId> UniqueClanID;
 	
 	// Holds the message type
-	EFriendsRequestType::Type MessageType;
+	EMessageType::Type MessageType;
 	
 	// Holds the requester name.
 	FString RequesterName;
diff --git a/Engine/Source/Developer/FriendsAndChat/Public/Models/IChatTabViewModel.h b/Engine/Source/Developer/FriendsAndChat/Public/Models/IChatTabViewModel.h
--- a/Engine/Source/Developer/FriendsAndChat/Public/Models/IChatTabViewModel.h
+++ b/Engine/Source/Developer/FriendsAndChat/Public/Models/IChatTabViewModel.h
@@ -18,7 +18,7 @@ public:
 	virtual const EChatMessageType::Type GetTabID() const = 0;
 
 	virtual const FSlateBrush* GetTabImage() const = 0;
-	virtual const TSharedPtr<FChatViewModel> GetChatViewModel() const = 0;
+	virtual const TSharedPtr<class FChatViewModel> GetChatViewModel() const = 0;
 
 	DECLARE_EVENT(IChatTabViewModel, FChatTabVisibilityChangedEvent)
 	virtual FChatTabVisibilityChangedEvent& OnTabVisibilityChanged() = 0;
diff --git a/Engine/Source/Developer/FriendsAndChat/Public/Models/IFriendItem.h b/Engine/Source/Developer/FriendsAndChat/Public/Models/IFriendItem.h
--- a/Engine/Source/Developer/FriendsAndChat/Public/Models/IFriendItem.h
+++ b/Engine/Source/Developer/FriendsAndChat/Public/Models/IFriendItem.h
@@ -127,6 +127,12 @@ public:
 	*/
 	virtual const FString GetClientName() const = 0;
 
+	/**
+	* Get last seen time
+	* @return The time
+	*/
+	virtual FDateTime GetLastSeen() const = 0;
+
 	/**
 	 * Get if the user is online.
 	 * @return The user online state.
@@ -167,7 +173,7 @@ public:
 	 * Obtain info needed to join a party for this friend item
 	 * @return party info if available or null
 	 */
-	virtual TSharedPtr<IOnlinePartyJoinInfo> GetPartyJoinInfo() const = 0;
+	virtual TSharedPtr<class IOnlinePartyJoinInfo> GetPartyJoinInfo() const = 0;
 
 	/**
 	 * Get the Unique ID.
