[INTEGRATE] merge UE4-LauncherDev > UE4 CL 2636730

diff --git a/Engine/Source/Developer/FriendsAndChat/Private/Layers/Presentation/ChatTextLayoutMarshaller.cpp b/Engine/Source/Developer/FriendsAndChat/Private/Layers/Presentation/ChatTextLayoutMarshaller.cpp
--- a/Engine/Source/Developer/FriendsAndChat/Private/Layers/Presentation/ChatTextLayoutMarshaller.cpp
+++ b/Engine/Source/Developer/FriendsAndChat/Private/Layers/Presentation/ChatTextLayoutMarshaller.cpp
@@ -112,7 +112,7 @@ public:
 
 			if(NameHyperlinkDecorator.IsValid() && (IsMultiChat || !NewMessage->IsFromSelf()))
 			{
-				FString MessageText = " " + GetSenderName(NewMessage) + " :";
+				FString MessageText = " " + GetSenderName(NewMessage) + " : ";
 
 				int32 NameLen = MessageText.Len();
 
@@ -143,18 +143,36 @@ public:
 			else
 			{
 				ModelRange.BeginIndex = ModelString->Len();
-				*ModelString += (" " + NewMessage->GetSenderName().ToString() + " :");
+				*ModelString += (" " + NewMessage->GetSenderName().ToString() + " : ");
 				ModelRange.EndIndex = ModelString->Len();
 				Runs.Add(FSlateTextRun::Create(FRunInfo(), ModelString, MessageTextStyle, ModelRange));
 			}
 		}
 
-		ModelRange.BeginIndex = ModelString->Len();
-		*ModelString += (" " + NewMessage->GetMessage().ToString());
-		ModelRange.EndIndex = ModelString->Len();
-		Runs.Add(FSlateTextRun::Create(FRunInfo(), ModelString, MessageTextStyle, ModelRange));
+		FString MessageString = NewMessage->GetMessage().ToString();
+		TArray<FString> MessageLines;
+		MessageString.ParseIntoArrayLines(MessageLines);
+
+		if(MessageLines.Num())
+		{
+			ModelRange.BeginIndex = ModelString->Len();
+			*ModelString += MessageLines[0];
+			ModelRange.EndIndex = ModelString->Len();
+			Runs.Add(FSlateTextRun::Create(FRunInfo(), ModelString, MessageTextStyle, ModelRange));
+			TextLayout->AddLine(ModelString, Runs);
+		}
+
+		for (int32 Range = 1; Range < MessageLines.Num(); Range++)
+		{
+			TSharedRef<FString> LineString = MakeShareable(new FString(MessageLines[Range]));
+			TArray< TSharedRef< IRun > > LineRun;
+			FTextRange LineRange;
+			LineRange.BeginIndex = 0;
+			LineRange.EndIndex = LineString->Len();
+			LineRun.Add(FSlateTextRun::Create(FRunInfo(), LineString, MessageTextStyle, LineRange));
+			TextLayout->AddLine(LineString, LineRun);
+		}
 
-		TextLayout->AddLine(ModelString, Runs);
 		return true;
 	}
 
