Fixes for a case-sensitive fs and shadowing warnings.

diff --git a/Engine/Source/Developer/FriendsAndChat/Private/Layers/Presentation/ChatNotificationService.cpp b/Engine/Source/Developer/FriendsAndChat/Private/Layers/Presentation/ChatNotificationService.cpp
--- a/Engine/Source/Developer/FriendsAndChat/Private/Layers/Presentation/ChatNotificationService.cpp
+++ b/Engine/Source/Developer/FriendsAndChat/Private/Layers/Presentation/ChatNotificationService.cpp
@@ -2,7 +2,7 @@
 
 #include "FriendsAndChatPrivatePCH.h"
 #include "ChatNotificationService.h"
-#include "Claninfo.h"
+#include "ClanInfo.h"
 
 #define LOCTEXT_NAMESPACE ""
 
@@ -146,4 +146,4 @@ TSharedRef< FChatNotificationService > FChatNotificationServiceFactory::Create()
 	return NotificationService;
 }
 
-#undef LOCTEXT_NAMESPACE
\ No newline at end of file
+#undef LOCTEXT_NAMESPACE
