Merging

//depot/UE4-Orion/Engine/Source/Developer/FriendsAndChat/Public/FriendsComboStyle.h

to //depot/UE4/Engine/Source/Developer/FriendsAndChat/Public/FriendsComboStyle.h

diff --git a/Engine/Source/Developer/FriendsAndChat/Public/FriendsComboStyle.h b/Engine/Source/Developer/FriendsAndChat/Public/FriendsComboStyle.h
--- a/Engine/Source/Developer/FriendsAndChat/Public/FriendsComboStyle.h
+++ b/Engine/Source/Developer/FriendsAndChat/Public/FriendsComboStyle.h
@@ -39,7 +39,7 @@ struct FRIENDSANDCHAT_API FFriendsComboStyle
 	 */
 	static const FFriendsComboStyle& GetDefault();
 
-	UPROPERTY()
+	UPROPERTY(EditAnywhere, Category = Appearance)
 	FButtonStyle ComboItemButtonStyle;
 	FFriendsComboStyle& SetComboItemButtonStyle(const FButtonStyle& ButtonStyle);
 
@@ -88,27 +88,27 @@ struct FRIENDSANDCHAT_API FFriendsComboStyle
 	FFriendsComboStyle& SetComboItemContentPadding(const FMargin& InPadding);
 
 	/** Friends List Combo Button menu background image (left) */
-	UPROPERTY()
+	UPROPERTY(EditAnywhere, Category = Appearance)
 	FSlateBrush FriendComboBackgroundLeftBrush;
 	FFriendsComboStyle& SetFriendComboBackgroundLeftBrush(const FSlateBrush& BrushStyle);
 
 	/** Friends List Combo Button menu background image (right) */
-	UPROPERTY()
+	UPROPERTY(EditAnywhere, Category = Appearance)
 	FSlateBrush FriendComboBackgroundRightBrush;
 	FFriendsComboStyle& SetFriendComboBackgroundRightBrush(const FSlateBrush& BrushStyle);
 
 	/** Friends List Combo Button menu background image (left-flipped) - for MenuPlacement_ComboBoxRight menus */
-	UPROPERTY()
+	UPROPERTY(EditAnywhere, Category = Appearance)
 	FSlateBrush FriendComboBackgroundLeftFlippedBrush;
 	FFriendsComboStyle& SetFriendComboBackgroundLeftFlippedBrush(const FSlateBrush& BrushStyle);
 
 	/** Friends List Combo Button menu background image (right-flipped) - for MenuPlacement_ComboBoxRight menus */
-	UPROPERTY()
+	UPROPERTY(EditAnywhere, Category = Appearance)
 	FSlateBrush FriendComboBackgroundRightFlippedBrush;
 	FFriendsComboStyle& SetFriendComboBackgroundRightFlippedBrush(const FSlateBrush& BrushStyle);
 
 	/** Friends List Combo Button style */
-	UPROPERTY()
+	UPROPERTY(EditAnywhere, Category = Appearance)
 	FComboButtonStyle FriendListComboButtonStyle;
 	FFriendsComboStyle& SetFriendsListComboButtonStyle(const FComboButtonStyle& ButtonStyle);
 };
