Merging 2665897, 2667359, 2667826

//depot/UE4-LauncherDev/...

to //depot/UE4/...

diff --git a/Engine/Source/Developer/FriendsAndChat/Private/FriendsChatStyle.cpp b/Engine/Source/Developer/FriendsAndChat/Private/FriendsChatStyle.cpp
--- a/Engine/Source/Developer/FriendsAndChat/Private/FriendsChatStyle.cpp
+++ b/Engine/Source/Developer/FriendsAndChat/Private/FriendsChatStyle.cpp
@@ -149,6 +149,12 @@ FFriendsChatStyle& FFriendsChatStyle::SetChatSettingsBrush(const FSlateBrush& Br
 	return *this;
 }
 
+FFriendsChatStyle& FFriendsChatStyle::SetMessageNotificationBrush(const FSlateBrush& Brush)
+{
+	MessageNotificationBrush = Brush;
+	return *this;
+}
+
 FFriendsChatStyle& FFriendsChatStyle::SetMinimizeButtonStyle(const FButtonStyle& Button)
 {
 	FriendsMinimizeButtonStyle = Button;
diff --git a/Engine/Source/Developer/FriendsAndChat/Private/Layers/Presentation/ChatChromeTabViewModel.cpp b/Engine/Source/Developer/FriendsAndChat/Private/Layers/Presentation/ChatChromeTabViewModel.cpp
--- a/Engine/Source/Developer/FriendsAndChat/Private/Layers/Presentation/ChatChromeTabViewModel.cpp
+++ b/Engine/Source/Developer/FriendsAndChat/Private/Layers/Presentation/ChatChromeTabViewModel.cpp
@@ -15,7 +15,7 @@ public:
 	{
 	}
 
-	bool IsTabVisible() const override
+	virtual bool IsTabVisible() const override
 	{
 		if (!ChatViewModel.IsValid())
 		{
@@ -35,22 +35,33 @@ public:
 		return true;
 	}
 
-	const FText GetTabText() const override
+	virtual const FText GetTabText() const override
 	{
 		return ChatViewModel->GetChatGroupText(false);
 	}
 
+	virtual FText GetMessageNotificationsText() const override
+	{
+		int32 Messages = ChatViewModel->GetUnreadChannelMessageCount();
+		if (Messages < 10)
+		{
+			return FText::AsNumber(Messages);
+		}
+		const FText Count = FText::AsNumber(9);
+		return FText::Format(NSLOCTEXT("","PlusNumber","{0}+"), Count);
+	}
+
 	virtual const EChatMessageType::Type GetTabID() const override
 	{
 		return ChatViewModel->GetChatChannelType();
 	}
 
-	const FSlateBrush* GetTabImage() const override
+	virtual const FSlateBrush* GetTabImage() const override
 	{
 		return nullptr;
 	}
 
-	const TSharedPtr<FChatViewModel> GetChatViewModel() const override
+	virtual const TSharedPtr<FChatViewModel> GetChatViewModel() const override
 	{
 		return ChatViewModel;
 	}
diff --git a/Engine/Source/Developer/FriendsAndChat/Private/Layers/Presentation/ChatChromeViewModel.cpp b/Engine/Source/Developer/FriendsAndChat/Private/Layers/Presentation/ChatChromeViewModel.cpp
--- a/Engine/Source/Developer/FriendsAndChat/Private/Layers/Presentation/ChatChromeViewModel.cpp
+++ b/Engine/Source/Developer/FriendsAndChat/Private/Layers/Presentation/ChatChromeViewModel.cpp
@@ -36,6 +36,16 @@ public:
 		}
 		ActiveTab = Tab;
 		ActiveTab->GetChatViewModel()->SetIsActive(true);
+
+		// Tell Other tabs what messages we read
+		ActiveTab->GetChatViewModel()->SetReadChannelFlags(0);
+		for (TSharedRef<IChatTabViewModel> Tab : Tabs)
+		{
+			if (Tab != ActiveTab)
+			{
+				Tab->GetChatViewModel()->SetReadChannelFlags(ActiveTab->GetChatViewModel()->GetChannelFlags());
+			}
+		}
 		
 		OnActiveTabChanged().Broadcast();
 	}
@@ -95,6 +105,14 @@ public:
 	virtual void ToggleChatMinimized() override
 	{
 		ChatDisplayService->ToggleChatMinimized();
+
+		if (IsChatMinimized())
+		{
+			for (TSharedRef<IChatTabViewModel> Tab : Tabs)
+			{
+				Tab->GetChatViewModel()->SetMessageShown(false, 0);
+			}
+		}
 	}
 
 	virtual bool IsChatMinimized() const override
diff --git a/Engine/Source/Developer/FriendsAndChat/Private/Layers/Presentation/ChatViewModel.cpp b/Engine/Source/Developer/FriendsAndChat/Private/Layers/Presentation/ChatViewModel.cpp
--- a/Engine/Source/Developer/FriendsAndChat/Private/Layers/Presentation/ChatViewModel.cpp
+++ b/Engine/Source/Developer/FriendsAndChat/Private/Layers/Presentation/ChatViewModel.cpp
@@ -54,6 +54,32 @@ public:
 		return FilteredMessages.Num();
 	}
 
+	virtual int32 GetUnreadChannelMessageCount() const override
+	{
+		int32 NewChannelMessageCount = 0;
+		if (!IsActive() || !bMessageShown)
+		{
+			for (int32 MessageIndex = LastSeenMessageCount; MessageIndex < GetMessageCount(); ++MessageIndex)
+			{
+				// Show Missed Messages if the active channel does now show our default message type.
+				if ((GetDefaultChannelType() & ReadChannelFlags) == 0)
+				{
+					// Show missed messages if we are the default channel for this message type or if we are the current active window
+					if (FilteredMessages[MessageIndex]->GetMessageType() == GetDefaultChannelType() || IsActive())
+					{
+						NewChannelMessageCount++;
+					}
+				}
+			}
+		}
+		return NewChannelMessageCount;
+	}
+
+	virtual void SetReadChannelFlags(uint8 ChannelFlags) override
+	{
+		ReadChannelFlags = ChannelFlags;
+	}
+
 	virtual TSharedPtr<FFriendViewModel> GetFriendViewModel(const TSharedPtr<const FUniqueNetId> InUserID, const FText Username) override
 	{
 		TSharedPtr<IFriendItem> FoundFriend = FriendsService->FindUser(InUserID.ToSharedRef());
@@ -106,6 +132,11 @@ public:
 		}
 	}
 
+	virtual uint8 GetChannelFlags() override
+	{
+		return ChatChannelFlags;
+	}
+
 	virtual bool IsChannelSet(const EChatMessageType::Type InChannel) override
 	{
 		return (InChannel & ChatChannelFlags) != 0;
@@ -598,6 +629,22 @@ public:
 			}
 			ChatDisplayService->SetActiveTab(SharedThis(this));
 		}
+		else
+		{
+			if (bMessageShown)
+			{
+				LastSeenMessageCount = GetMessageCount();
+			}
+		}
+	}
+
+	virtual void SetMessageShown(bool Shown, int32 NumMissedMessages) override
+	{
+		if (bMessageShown == true && Shown == false && IsActive() == true)
+		{
+			LastSeenMessageCount = GetMessageCount() - NumMissedMessages;
+		}
+		bMessageShown = Shown;
 	}
 
 	virtual bool AllowMarkup() override
@@ -770,6 +817,9 @@ protected:
 		, bAllowGlobalChat(true)
 		, bCaptureFocus(false)
 		, WindowOpacity(1.0f)
+		, LastSeenMessageCount(0)
+		, bMessageShown(true)
+		, ReadChannelFlags(0)
 	{
 	}
 
@@ -839,6 +889,10 @@ private:
 	EVisibility ChatEntryVisibility;
 	FSlateColor OverrideColor;
 
+	int32 LastSeenMessageCount;
+	bool bMessageShown;
+	uint8 ReadChannelFlags;
+
 private:
 	friend FChatViewModelFactory;
 };
diff --git a/Engine/Source/Developer/FriendsAndChat/Private/Layers/Presentation/ChatViewModel.h b/Engine/Source/Developer/FriendsAndChat/Private/Layers/Presentation/ChatViewModel.h
--- a/Engine/Source/Developer/FriendsAndChat/Private/Layers/Presentation/ChatViewModel.h
+++ b/Engine/Source/Developer/FriendsAndChat/Private/Layers/Presentation/ChatViewModel.h
@@ -45,6 +45,8 @@ public:
 	virtual bool SendMessage(const FText NewMessage, const FText PlainText) = 0;
 	virtual TArray<TSharedRef<FChatItemViewModel > >& GetMessages() = 0;
 	virtual int32 GetMessageCount() const = 0;
+	virtual int32 GetUnreadChannelMessageCount() const = 0;
+	virtual void SetReadChannelFlags(uint8 ChannelFlags) = 0;
 
 	// Friend Actions
 	virtual TSharedPtr<FFriendViewModel> GetFriendViewModel(const TSharedPtr<const FUniqueNetId> UniqueID, const FText Username) = 0;
@@ -52,6 +54,7 @@ public:
 	
 	// Channel
 	virtual void SetChannelFlags(uint8 ChannelFlags) = 0;
+	virtual uint8 GetChannelFlags() = 0;
 	virtual void SetDefaultChannelFlags(uint8 ChannelFlags) = 0;
 	virtual void SetDefaultOutgoingChannel(EChatMessageType::Type InChannel) = 0;
 	virtual void ResetToDefaultChannel() = 0;
@@ -92,6 +95,7 @@ public:
 	virtual void ValidateChatInput(const FText Message, const FText PlainText) = 0;
 	virtual FText GetValidatedInput() = 0;
 	virtual void SetIsActive(bool IsActive) = 0;
+	virtual void SetMessageShown(bool Shown, int32 NumMissedMessages = 1) = 0;
 	virtual void RefreshMessages() = 0;
 	virtual bool AllowMarkup() = 0;
 	virtual bool MultiChat() = 0;
diff --git a/Engine/Source/Developer/FriendsAndChat/Private/Layers/UI/SChatChrome.cpp b/Engine/Source/Developer/FriendsAndChat/Private/Layers/UI/SChatChrome.cpp
--- a/Engine/Source/Developer/FriendsAndChat/Private/Layers/UI/SChatChrome.cpp
+++ b/Engine/Source/Developer/FriendsAndChat/Private/Layers/UI/SChatChrome.cpp
@@ -184,6 +184,8 @@ private:
 
 			// Create container for a quick settings options
 			TSharedPtr<SBox> QuickSettingsContainer = SNew(SBox);
+
+			TSharedPtr<SBox> NotificationContainer = SNew(SBox);
 			
 			// Add Tab
 			TabContainer->AddSlot()
@@ -205,16 +207,25 @@ private:
 							TabButton
 						]
 						+ SHorizontalBox::Slot()
-							.AutoWidth()
+						.AutoWidth()
 						[
 							QuickSettingsContainer.ToSharedRef()
 						]
+						+ SHorizontalBox::Slot()
+						.AutoWidth()
+						[
+							NotificationContainer.ToSharedRef()
+						]
 					]
 				]
 			];
 
 			// Add a quick settings if needed
 			QuickSettingsContainer->SetContent(GetQuickSettingsWidget(Tab).ToSharedRef());
+
+			// Add Message notification indicator
+			NotificationContainer->SetContent(GetNotificationWidget(Tab).ToSharedRef());
+			
 			
 			if (!ChatWindows.Contains(Tab))
 			{
@@ -267,6 +278,47 @@ private:
 			];
 	}
 
+	EVisibility GetMessageNotificationVisibility(TWeakPtr<IChatTabViewModel> TabPtr) const
+	{
+		TSharedPtr<IChatTabViewModel> Tab = TabPtr.Pin();
+		if (Tab.IsValid())
+		{
+			if (Tab->GetChatViewModel()->GetUnreadChannelMessageCount() > 0)
+			{
+				return EVisibility::Visible;
+			}
+		}
+		return EVisibility::Collapsed;
+	}
+
+	TSharedPtr<SWidget> GetNotificationWidget(const TSharedRef<IChatTabViewModel>& Tab) const
+	{
+		return SNew(SBox)
+		.Visibility(this, &SChatChromeImpl::GetMessageNotificationVisibility, TWeakPtr<IChatTabViewModel>(Tab))
+		[
+			SNew(SBorder)
+			.Padding(FMargin(4.0f))
+			.BorderImage(&FriendStyle.FriendsChatStyle.MessageNotificationBrush)
+			[
+				SNew(STextBlock)
+				.ColorAndOpacity(FriendStyle.FriendsNormalFontStyle.DefaultFontColor)
+				.Font(FriendStyle.FriendsNormalFontStyle.FriendsFontSmallBold)
+				.Text(this, &SChatChromeImpl::GetMessageNotificationsText, TWeakPtr<IChatTabViewModel>(Tab))
+			]
+		];
+	}
+
+	FText GetMessageNotificationsText(TWeakPtr<IChatTabViewModel> TabPtr) const
+	{
+		TSharedPtr<IChatTabViewModel> Tab = TabPtr.Pin();
+		if (Tab.IsValid())
+		{
+			return Tab->GetMessageNotificationsText();
+		}
+		return FText::GetEmpty();
+	}
+
+
 	EVisibility GetChatWindowVisibility() const
 	{
 		return ChromeViewModel->IsActive() ? EVisibility::SelfHitTestInvisible : EVisibility::Collapsed;
diff --git a/Engine/Source/Developer/FriendsAndChat/Private/Layers/UI/SChatWindow.cpp b/Engine/Source/Developer/FriendsAndChat/Private/Layers/UI/SChatWindow.cpp
--- a/Engine/Source/Developer/FriendsAndChat/Private/Layers/UI/SChatWindow.cpp
+++ b/Engine/Source/Developer/FriendsAndChat/Private/Layers/UI/SChatWindow.cpp
@@ -221,6 +221,14 @@ public:
 		{
 			bUserHasScrolled = true;
 		}
+		if (bUserHasScrolled && ExternalScrollbar.IsValid() && !ExternalScrollbar->IsScrolling())
+		{
+			bUserHasScrolled = false;
+			if (ExternalScrollbar->DistanceFromBottom() == 0)
+			{
+				ViewModel->SetMessageShown(true);
+			}
+		}
 	}
 
 	virtual FReply OnMouseWheel( const FGeometry& MyGeometry, const FPointerEvent& MouseEvent ) override
@@ -310,6 +318,11 @@ private:
 			if(ExternalScrollbar.IsValid() && ExternalScrollbar->DistanceFromBottom() == 0)
 			{
 				ChatScrollBox->ScrollToEnd();
+				ViewModel->SetMessageShown(true);
+			}
+			else
+			{
+				ViewModel->SetMessageShown(false);
 			}
 		}
 	}
diff --git a/Engine/Source/Developer/FriendsAndChat/Public/FriendsChatStyle.h b/Engine/Source/Developer/FriendsAndChat/Public/FriendsChatStyle.h
--- a/Engine/Source/Developer/FriendsAndChat/Public/FriendsChatStyle.h
+++ b/Engine/Source/Developer/FriendsAndChat/Public/FriendsChatStyle.h
@@ -135,6 +135,10 @@ struct FRIENDSANDCHAT_API FFriendsChatStyle
 	FSlateBrush ChatSettingsBrush;
 	FFriendsChatStyle& SetChatSettingsBrush(const FSlateBrush& Brush);
 
+	UPROPERTY(EditAnywhere, Category = Appearance)
+	FSlateBrush MessageNotificationBrush;
+	FFriendsChatStyle& SetMessageNotificationBrush(const FSlateBrush& Brush);
+
 	UPROPERTY()
 	FButtonStyle FriendsMinimizeButtonStyle;
 	FFriendsChatStyle& SetMinimizeButtonStyle(const FButtonStyle& ButtonStyle);
diff --git a/Engine/Source/Developer/FriendsAndChat/Public/Models/IChatTabViewModel.h b/Engine/Source/Developer/FriendsAndChat/Public/Models/IChatTabViewModel.h
--- a/Engine/Source/Developer/FriendsAndChat/Public/Models/IChatTabViewModel.h
+++ b/Engine/Source/Developer/FriendsAndChat/Public/Models/IChatTabViewModel.h
@@ -9,6 +9,7 @@ public:
 	// IChatTabViewModel Interface
 	virtual bool IsTabVisible() const = 0;
 	virtual const FText GetTabText() const = 0;
+	virtual FText GetMessageNotificationsText() const = 0;
 
 	/**
 	 * Provides an ID Used in Chat Window navigation.
