Fix CIS shadowing warning
Fix naming convention/consistency on OverrideDisplayVisibility

diff --git a/Engine/Source/Developer/FriendsAndChat/Private/Layers/Presentation/ChatChromeViewModel.cpp b/Engine/Source/Developer/FriendsAndChat/Private/Layers/Presentation/ChatChromeViewModel.cpp
--- a/Engine/Source/Developer/FriendsAndChat/Private/Layers/Presentation/ChatChromeViewModel.cpp
+++ b/Engine/Source/Developer/FriendsAndChat/Private/Layers/Presentation/ChatChromeViewModel.cpp
@@ -29,13 +29,13 @@ public:
 		}
 	}
 
-	virtual void ActivateTab(const TSharedRef<IChatTabViewModel>& Tab) override
+	virtual void ActivateTab(const TSharedRef<IChatTabViewModel>& InTab) override
 	{
 		if(ActiveTab.IsValid())
 		{
 			ActiveTab->GetChatViewModel()->SetIsActive(false);
 		}
-		ActiveTab = Tab;
+		ActiveTab = InTab;
 		ActiveTab->GetChatViewModel()->SetIsActive(true);
 
 		// Tell Other tabs what messages we read
diff --git a/Engine/Source/Developer/FriendsAndChat/Private/Layers/Presentation/ChatViewModel.cpp b/Engine/Source/Developer/FriendsAndChat/Private/Layers/Presentation/ChatViewModel.cpp
--- a/Engine/Source/Developer/FriendsAndChat/Private/Layers/Presentation/ChatViewModel.cpp
+++ b/Engine/Source/Developer/FriendsAndChat/Private/Layers/Presentation/ChatViewModel.cpp
@@ -386,7 +386,7 @@ public:
 
 	virtual bool IsOverrideDisplaySet() override
 	{
-		return OverrideDisplayVisibility;
+		return bOverrideDisplayVisibility;
 	}
 
 	// End FChatViewModel interface
@@ -674,7 +674,7 @@ public:
 	virtual void SetIsActive(bool InIsActive)
 	{
 		bIsActive = InIsActive;
-		OverrideDisplayVisibility = true;
+		bOverrideDisplayVisibility = true;
 
 		// If active, ensure we have a valid chat channel
 		if(bIsActive)
@@ -905,7 +905,7 @@ protected:
 		, bInParty(false)
 		, bHasActionPending(false)
 		, bAllowJoinGame(false)
-		, OverrideDisplayVisibility(false)
+		, bOverrideDisplayVisibility(false)
 		, bAllowFade(true)
 		, bUseOverrideColor(false)
 		, bAllowGlobalChat(true)
@@ -961,7 +961,7 @@ private:
 	bool bInParty;
 	bool bHasActionPending;
 	bool bAllowJoinGame;
-	bool OverrideDisplayVisibility;
+	bool bOverrideDisplayVisibility;
 
 	TArray<TSharedRef<FChatItemViewModel> > FilteredMessages;
 	TArray<TSharedPtr<FSelectedFriend> > RecentPlayerList;
