Follow up to remove unneeded #include of internal MCP code

diff --git a/Engine/Source/Developer/FriendsAndChat/Private/Layers/DataAccess/FriendsAndChatManager.cpp b/Engine/Source/Developer/FriendsAndChat/Private/Layers/DataAccess/FriendsAndChatManager.cpp
--- a/Engine/Source/Developer/FriendsAndChat/Private/Layers/DataAccess/FriendsAndChatManager.cpp
+++ b/Engine/Source/Developer/FriendsAndChat/Private/Layers/DataAccess/FriendsAndChatManager.cpp
@@ -33,8 +33,6 @@
 #include "OSSScheduler.h"
 #include "FriendsService.h"
 #include "GameAndPartyService.h"
-#include "OnlineSubsystemMcp.h"
-#include "GameServiceMcp.h"
 
 #define LOCTEXT_NAMESPACE "FriendsAndChatManager"
 
