Merging latest social fixed from Orion to Main

//depot/UE4-Orion/Engine/Source/Developer/FriendsAndChat/...

to //depot/UE4/Engine/Source/Developer/FriendsAndChat/...

diff --git a/Engine/Source/Developer/FriendsAndChat/Private/FriendsAndChatModule.cpp b/Engine/Source/Developer/FriendsAndChat/Private/FriendsAndChatModule.cpp
--- a/Engine/Source/Developer/FriendsAndChat/Private/FriendsAndChatModule.cpp
+++ b/Engine/Source/Developer/FriendsAndChat/Private/FriendsAndChatModule.cpp
@@ -2,6 +2,8 @@
 
 #include "FriendsAndChatPrivatePCH.h"
 
+#include "FriendsAndChatStyle.h"
+
 /**
  * Implements the FriendsAndChat module.
  */
@@ -38,6 +40,11 @@ public:
 		return NewManager;
 	}
 
+	virtual void ShutdownStyle() override
+	{
+		FFriendsAndChatModuleStyle::Shutdown();
+	}
+
 public:
 
 	// IModuleInterface interface
diff --git a/Engine/Source/Developer/FriendsAndChat/Private/FriendsAndChatStyle.cpp b/Engine/Source/Developer/FriendsAndChat/Private/FriendsAndChatStyle.cpp
--- a/Engine/Source/Developer/FriendsAndChat/Private/FriendsAndChatStyle.cpp
+++ b/Engine/Source/Developer/FriendsAndChat/Private/FriendsAndChatStyle.cpp
@@ -109,9 +109,12 @@ void FFriendsAndChatModuleStyle::Initialize(FFriendsAndChatStyle FriendStyle)
 
 void FFriendsAndChatModuleStyle::Shutdown()
 {
-	FSlateStyleRegistry::UnRegisterSlateStyle( *FriendsAndChatModuleStyleInstance );
-	ensure( FriendsAndChatModuleStyleInstance.IsUnique() );
-	FriendsAndChatModuleStyleInstance.Reset();
+	if ( FriendsAndChatModuleStyleInstance.IsValid() )
+	{
+		FSlateStyleRegistry::UnRegisterSlateStyle( *FriendsAndChatModuleStyleInstance );
+		ensure( FriendsAndChatModuleStyleInstance.IsUnique() );
+		FriendsAndChatModuleStyleInstance.Reset();
+	}
 }
 
 FName FFriendsAndChatModuleStyle::GetStyleSetName()
diff --git a/Engine/Source/Developer/FriendsAndChat/Public/FriendsAndChatStyle.h b/Engine/Source/Developer/FriendsAndChat/Public/FriendsAndChatStyle.h
--- a/Engine/Source/Developer/FriendsAndChat/Public/FriendsAndChatStyle.h
+++ b/Engine/Source/Developer/FriendsAndChat/Public/FriendsAndChatStyle.h
@@ -74,31 +74,31 @@ struct FRIENDSANDCHAT_API FFriendsAndChatStyle
 	FFriendsAndChatStyle& SetWindowStyle(const FWindowStyle& InStyle);
 
 	/** SFriendActions Action Button style */
-	UPROPERTY()
+	UPROPERTY( EditAnywhere, Category = Appearance )
 	FButtonStyle ActionButtonStyle;
 	FFriendsAndChatStyle& SetActionButtonStyle(const FButtonStyle& ButtonStyle);
 
-	UPROPERTY()
+	UPROPERTY( EditAnywhere, Category = Appearance )
 	FFriendsFontStyle FriendsNormalFontStyle;
 	FFriendsAndChatStyle& SetNormalFriendsFontStyle(const FFriendsFontStyle& FontStyle);
 
-	UPROPERTY()
+	UPROPERTY( EditAnywhere, Category = Appearance )
 	FFriendsListStyle FriendsListStyle;
 	FFriendsAndChatStyle& SetFriendsListStyle(const FFriendsListStyle& InFriendsListStyle);
 
-	UPROPERTY()
+	UPROPERTY( EditAnywhere, Category = Appearance )
 	FFriendsComboStyle FriendsComboStyle;
 	FFriendsAndChatStyle& SetFriendsComboStyle(const FFriendsComboStyle& InFriendsListStyle);
 
-	UPROPERTY()
+	UPROPERTY( EditAnywhere, Category = Appearance )
 	FFriendsChatStyle FriendsChatStyle;
 	FFriendsAndChatStyle& SetFriendsChatStyle(const FFriendsChatStyle& InFriendsChatStyle);
 
-	UPROPERTY()
+	UPROPERTY( EditAnywhere, Category = Appearance )
 	FFriendsChatChromeStyle FriendsChatChromeStyle;
 	FFriendsAndChatStyle& SetFriendsChatChromeStyle(const FFriendsChatChromeStyle& InFriendsChatChromeStyle);
 
-	UPROPERTY()
+	UPROPERTY( EditAnywhere, Category = Appearance )
 	FFriendsMarkupStyle FriendsMarkupStyle;
 	FFriendsAndChatStyle& SetFriendsMarkupStyle(const FFriendsMarkupStyle& InFriendsMarkupStyle);
 };
diff --git a/Engine/Source/Developer/FriendsAndChat/Public/FriendsChatStyle.h b/Engine/Source/Developer/FriendsAndChat/Public/FriendsChatStyle.h
--- a/Engine/Source/Developer/FriendsAndChat/Public/FriendsChatStyle.h
+++ b/Engine/Source/Developer/FriendsAndChat/Public/FriendsChatStyle.h
@@ -39,7 +39,7 @@ struct FRIENDSANDCHAT_API FFriendsChatStyle
 	 */
 	static const FFriendsChatStyle& GetDefault();
 
-	UPROPERTY()
+	UPROPERTY(EditAnywhere, Category = Appearance)
 	FSlateBrush GlobalChatHeaderBrush;
 	FFriendsChatStyle& SetGlobalChatHeaderBrush(const FSlateBrush& Value);
 
@@ -83,7 +83,7 @@ struct FRIENDSANDCHAT_API FFriendsChatStyle
 	FLinearColor PartyChatColor;
 	FFriendsChatStyle& SetPartyChatColor(const FLinearColor& InFontColor);
 
-	UPROPERTY(EditAnywhere, Category = Appearance)
+	UPROPERTY(EditAnywhere, Category=Appearance)
 	FLinearColor AdminChatColor;
 	FFriendsChatStyle& SetAdminChatColor(const FLinearColor& InFontColor);
 
@@ -163,7 +163,7 @@ struct FRIENDSANDCHAT_API FFriendsChatStyle
 	FButtonStyle FriendsMinimizeButtonStyle;
 	FFriendsChatStyle& SetMinimizeButtonStyle(const FButtonStyle& ButtonStyle);
 
-	UPROPERTY()
+	UPROPERTY(EditAnywhere, Category = Appearance)
 	FButtonStyle FriendsMaximizeButtonStyle;
 	FFriendsChatStyle& SetMaximizeButtonStyle(const FButtonStyle& ButtonStyle);
 
diff --git a/Engine/Source/Developer/FriendsAndChat/Public/FriendsListStyle.h b/Engine/Source/Developer/FriendsAndChat/Public/FriendsListStyle.h
--- a/Engine/Source/Developer/FriendsAndChat/Public/FriendsListStyle.h
+++ b/Engine/Source/Developer/FriendsAndChat/Public/FriendsListStyle.h
@@ -44,15 +44,15 @@ struct FRIENDSANDCHAT_API FFriendsListStyle
 
 	// Friends List Style
 
-	UPROPERTY()
+	UPROPERTY( EditAnywhere, Category = Appearance )
 	FFriendsFontStyle FriendsNormalFontStyle;
 	FFriendsListStyle& SetNormalFriendsFontStyle(const FFriendsFontStyle& FontStyle);
 
-	UPROPERTY()
+	UPROPERTY( EditAnywhere, Category = Appearance )
 	FButtonStyle GlobalChatButtonStyle;
 	FFriendsListStyle& SetGlobalChatButtonStyle(const FButtonStyle& ButtonStyle);
 
-	UPROPERTY()
+	UPROPERTY( EditAnywhere, Category = Appearance )
 	FFriendsComboStyle FriendsComboStyle;
 	FFriendsListStyle& SetFriendsComboStyle(const FFriendsComboStyle& FriendsComboStyle);
 
@@ -60,11 +60,11 @@ struct FRIENDSANDCHAT_API FFriendsListStyle
 	FScrollBarStyle ScrollBarStyle;
 	FFriendsListStyle& SetScrollbarStyle(const FScrollBarStyle& InScrollBarStyle);
 
-	UPROPERTY()
+	UPROPERTY( EditAnywhere, Category = Appearance )
 	FSlateColor ButtonInvertedForegroundColor;
 	FFriendsListStyle& SetButtonInvertedForegroundColor(const FSlateColor& Value);
 
-	UPROPERTY()
+	UPROPERTY( EditAnywhere, Category = Appearance )
 	FSlateColor ButtonForegroundColor;
 	FFriendsListStyle& SetButtonForegroundColor(const FSlateColor& Value);
 
@@ -72,60 +72,60 @@ struct FRIENDSANDCHAT_API FFriendsListStyle
 	FMargin BorderPadding;
 	FFriendsListStyle& SetBorderPadding(const FMargin& InBorderPadding);
 
-	UPROPERTY()
+	UPROPERTY( EditAnywhere, Category = Appearance )
 	FFriendsFontStyle NormalFont;
 	FFriendsListStyle& SetNormalFont(const FFriendsFontStyle& FontStyle);
 
 	/** Friends List Open Button style */
-	UPROPERTY()
+	UPROPERTY( EditAnywhere, Category = Appearance )
 	FButtonStyle FriendListOpenButtonStyle;
 	FFriendsListStyle& SetFriendsListOpenButtonStyle(const FButtonStyle& ButtonStyle);
 
 	/** Friends General Purpose Button style */
-	UPROPERTY()
+	UPROPERTY( EditAnywhere, Category = Appearance )
 	FButtonStyle FriendGeneralButtonStyle;
 	FFriendsListStyle& SetFriendGeneralButtonStyle(const FButtonStyle& ButtonStyle);
 
 	/** Friends List Action Button style */
-	UPROPERTY()
+	UPROPERTY( EditAnywhere, Category = Appearance )
 	FButtonStyle FriendListActionButtonStyle;
 	FFriendsListStyle& SetFriendListActionButtonStyle(const FButtonStyle& ButtonStyle);
 
 	/** Friends List Critical Button style */
-	UPROPERTY()
+	UPROPERTY( EditAnywhere, Category = Appearance )
 	FButtonStyle FriendListCriticalButtonStyle;
 	FFriendsListStyle& SetFriendsListCriticalButtonStyle(const FButtonStyle& ButtonStyle);
 
 	/** Friends List Emphasis Button style */
-	UPROPERTY()
+	UPROPERTY( EditAnywhere, Category = Appearance )
 	FButtonStyle FriendListEmphasisButtonStyle;
 	FFriendsListStyle& SetFriendsListEmphasisButtonStyle(const FButtonStyle& ButtonStyle);
 
-	UPROPERTY()
+	UPROPERTY( EditAnywhere, Category = Appearance )
 	FButtonStyle FriendListItemButtonStyle;
 	FFriendsListStyle& SetFriendsListItemButtonStyle(const FButtonStyle& ButtonStyle);
 
-	UPROPERTY()
+	UPROPERTY( EditAnywhere, Category = Appearance )
 	FButtonStyle FriendListItemButtonSimpleStyle;
 	FFriendsListStyle& SetFriendsListItemButtonSimpleStyle(const FButtonStyle& ButtonStyle);
 
 	/** Friends List Close button style */
-	UPROPERTY()
+	UPROPERTY( EditAnywhere, Category = Appearance )
 	FButtonStyle FriendListCloseButtonStyle;
 	FFriendsListStyle& SetFriendsListClosedButtonStyle(const FButtonStyle& ButtonStyle);
 
 	/** Add Friend Close button style */
-	UPROPERTY()
+	UPROPERTY( EditAnywhere, Category = Appearance )
 	FButtonStyle AddFriendCloseButtonStyle;
 	FFriendsListStyle& SetAddFriendCloseButtonStyle(const FButtonStyle& ButtonStyle);
 
 	/** Optional content for the Add Friend button */
-	UPROPERTY()
+	UPROPERTY( EditAnywhere, Category = Appearance )
 	FSlateBrush AddFriendButtonContentBrush;
 	FFriendsListStyle& SetAddFriendButtonContentBrush(const FSlateBrush& BrushStyle);
 
 	/** Optional content for the Add Friend button (hovered) */
-	UPROPERTY()
+	UPROPERTY( EditAnywhere, Category = Appearance )
 	FSlateBrush AddFriendButtonContentHoveredBrush;
 	FFriendsListStyle& SetAddFriendButtonContentHoveredBrush(const FSlateBrush& BrushStyle);
 
diff --git a/Engine/Source/Developer/FriendsAndChat/Public/Interfaces/IFriendsAndChatModule.h b/Engine/Source/Developer/FriendsAndChat/Public/Interfaces/IFriendsAndChatModule.h
--- a/Engine/Source/Developer/FriendsAndChat/Public/Interfaces/IFriendsAndChatModule.h
+++ b/Engine/Source/Developer/FriendsAndChat/Public/Interfaces/IFriendsAndChatModule.h
@@ -25,6 +25,8 @@ public:
 
 	virtual TSharedRef<IFriendsAndChatManager> GetFriendsAndChatManager(FName OSSInstanceName = TEXT(""), bool InGame = false) = 0;
 
+	virtual void ShutdownStyle() = 0;
+
 public:
 
 	/** Virtual destructor. */
