Fix to global chat state not being reset to "unavailable" on logout
Merging //depot/UE4-LauncherDev/Engine/Source/Developer/FriendsAndChat/Private/Layers/Domain/MessageService.cpp to //depot/UE4/Engine/Source/Developer/FriendsAndChat/Private/Layers/Domain/MessageService.cpp

diff --git a/Engine/Source/Developer/FriendsAndChat/Private/Layers/Domain/MessageService.cpp b/Engine/Source/Developer/FriendsAndChat/Private/Layers/Domain/MessageService.cpp
--- a/Engine/Source/Developer/FriendsAndChat/Private/Layers/Domain/MessageService.cpp
+++ b/Engine/Source/Developer/FriendsAndChat/Private/Layers/Domain/MessageService.cpp
@@ -60,6 +60,7 @@ public:
 		}
 		RoomJoins.Empty();
 		LoggedInUser.Reset();
+		bJoinedGlobalChat = false;
 		UnInitialize();
 		bJoinedGlobalChat = false;
 	}
@@ -243,6 +244,7 @@ private:
 		GameMessagesCount = 0;
 		PartyMessagesCount = 0;
 		ReceivedMessages.Empty();
+		bJoinedGlobalChat = false;
 
 		IOnlineChatPtr ChatInterface = OSSScheduler->GetChatInterface();
 		if (ChatInterface.IsValid())
