Nonunity fix to FriendsAndChatStyle needing ChatMessageType enum forward declared

diff --git a/Engine/Source/Developer/FriendsAndChat/Public/FriendsChatStyle.h b/Engine/Source/Developer/FriendsAndChat/Public/FriendsChatStyle.h
--- a/Engine/Source/Developer/FriendsAndChat/Public/FriendsChatStyle.h
+++ b/Engine/Source/Developer/FriendsAndChat/Public/FriendsChatStyle.h
@@ -4,6 +4,12 @@
 
 #include "FriendsChatStyle.generated.h"
 
+// Forward declarations
+namespace EChatMessageType
+{
+	enum Type : uint8;
+}
+
 /**
  * Interface for the services manager.
  */
