Integrate 2687864 from UE4-Orion

- Add player count tracking for PS4 sessions in Register/Unregister players
- Creating the MCP party join token doesn't add -party_joininfo_token= by default, this is added by the function that calls it now

diff --git a/Engine/Source/Developer/FriendsAndChat/Private/Layers/Domain/GameAndPartyService.cpp b/Engine/Source/Developer/FriendsAndChat/Private/Layers/Domain/GameAndPartyService.cpp
--- a/Engine/Source/Developer/FriendsAndChat/Private/Layers/Domain/GameAndPartyService.cpp
+++ b/Engine/Source/Developer/FriendsAndChat/Private/Layers/Domain/GameAndPartyService.cpp
@@ -476,7 +476,7 @@ private:
 			// notify for further processing of join party request 
 			OnFriendsJoinParty().Broadcast(*FriendItem->GetUniqueID(), PartyInfo.ToSharedRef(), true);
 
-			AdditionalCommandline = OSSScheduler->GetPartyInterface()->MakeTokenFromJoinInfo(*PartyInfo);
+			AdditionalCommandline = TEXT("-party_joininfo_token=") + OSSScheduler->GetPartyInterface()->MakeTokenFromJoinInfo(*PartyInfo);
 		}
 		else if(FriendItem.IsValid() && FriendItem->GetGameSessionId().IsValid())
 		{
